Fixing references and cleaning up naming in level prototyping kit.
Adding aseets to third person BP project.

diff --git a/Templates/TP_ThirdPersonBP/Config/TemplateDefs.ini b/Templates/TP_ThirdPersonBP/Config/TemplateDefs.ini
--- a/Templates/TP_ThirdPersonBP/Config/TemplateDefs.ini
+++ b/Templates/TP_ThirdPersonBP/Config/TemplateDefs.ini
@@ -44,4 +44,5 @@ ReplacementsInFiles=(Extensions=("cpp","h","ini","cs"),From="%TEMPLATENAME%",To=
 SharedContentPacks=(MountName="ThirdPerson",DetailLevels=("Standard"))
 SharedContentPacks=(MountName="Geometry",DetailLevels=("Standard","High"))
 SharedContentPacks=(MountName="Mannequin",DetailLevels=("Standard","High"))
+SharedContentPacks=(MountName="LevelPrototyping",DetailLevels=("High"))
 EditDetailLevelPreference="High"
