Third Person Template Overhaul + Shared Mannequin Resource Changes

diff --git a/Templates/TP_ThirdPerson/Config/DefaultEngine.ini b/Templates/TP_ThirdPerson/Config/DefaultEngine.ini
--- a/Templates/TP_ThirdPerson/Config/DefaultEngine.ini
+++ b/Templates/TP_ThirdPerson/Config/DefaultEngine.ini
@@ -1,7 +1,8 @@
 [/Script/EngineSettings.GameMapsSettings]
-GameDefaultMap=/Game/ThirdPersonCPP/Maps/ThirdPersonExampleMap
-EditorStartupMap=/Game/ThirdPersonCPP/Maps/ThirdPersonExampleMap
+GameDefaultMap=/Game/ThirdPerson/Maps/ThirdPersonMap.ThirdPersonMap
+EditorStartupMap=/Game/ThirdPerson/Maps/ThirdPersonMap.ThirdPersonMap
 GlobalDefaultGameMode="/Script/TP_ThirdPerson.TP_ThirdPersonGameMode"
 
 [/Script/IOSRuntimeSettings.IOSRuntimeSettings]
 MinimumiOSVersion=IOS_13
+
diff --git a/Templates/TP_ThirdPerson/Config/DefaultInput.ini b/Templates/TP_ThirdPerson/Config/DefaultInput.ini
--- a/Templates/TP_ThirdPerson/Config/DefaultInput.ini
+++ b/Templates/TP_ThirdPerson/Config/DefaultInput.ini
@@ -9,89 +9,8 @@
 +AxisConfig=(AxisKeyName="MouseWheelAxis",AxisProperties=(DeadZone=0.000000,Sensitivity=1.000000,Exponent=1.000000,bInvert=False))
 +AxisConfig=(AxisKeyName="Gamepad_LeftTriggerAxis",AxisProperties=(DeadZone=0.000000,Sensitivity=1.000000,Exponent=1.000000,bInvert=False))
 +AxisConfig=(AxisKeyName="Gamepad_RightTriggerAxis",AxisProperties=(DeadZone=0.000000,Sensitivity=1.000000,Exponent=1.000000,bInvert=False))
-+AxisConfig=(AxisKeyName="MotionController_Left_Thumbstick_X",AxisProperties=(DeadZone=0.000000,Sensitivity=1.000000,Exponent=1.000000,bInvert=False))
-+AxisConfig=(AxisKeyName="MotionController_Left_Thumbstick_Y",AxisProperties=(DeadZone=0.000000,Sensitivity=1.000000,Exponent=1.000000,bInvert=False))
-+AxisConfig=(AxisKeyName="MotionController_Left_TriggerAxis",AxisProperties=(DeadZone=0.000000,Sensitivity=1.000000,Exponent=1.000000,bInvert=False))
-+AxisConfig=(AxisKeyName="MotionController_Left_Grip1Axis",AxisProperties=(DeadZone=0.000000,Sensitivity=1.000000,Exponent=1.000000,bInvert=False))
-+AxisConfig=(AxisKeyName="MotionController_Left_Grip2Axis",AxisProperties=(DeadZone=0.000000,Sensitivity=1.000000,Exponent=1.000000,bInvert=False))
-+AxisConfig=(AxisKeyName="MotionController_Right_Thumbstick_X",AxisProperties=(DeadZone=0.000000,Sensitivity=1.000000,Exponent=1.000000,bInvert=False))
-+AxisConfig=(AxisKeyName="MotionController_Right_Thumbstick_Y",AxisProperties=(DeadZone=0.000000,Sensitivity=1.000000,Exponent=1.000000,bInvert=False))
-+AxisConfig=(AxisKeyName="MotionController_Right_TriggerAxis",AxisProperties=(DeadZone=0.000000,Sensitivity=1.000000,Exponent=1.000000,bInvert=False))
-+AxisConfig=(AxisKeyName="MotionController_Right_Grip1Axis",AxisProperties=(DeadZone=0.000000,Sensitivity=1.000000,Exponent=1.000000,bInvert=False))
-+AxisConfig=(AxisKeyName="MotionController_Right_Grip2Axis",AxisProperties=(DeadZone=0.000000,Sensitivity=1.000000,Exponent=1.000000,bInvert=False))
 +AxisConfig=(AxisKeyName="Gamepad_Special_Left_X",AxisProperties=(DeadZone=0.000000,Sensitivity=1.000000,Exponent=1.000000,bInvert=False))
 +AxisConfig=(AxisKeyName="Gamepad_Special_Left_Y",AxisProperties=(DeadZone=0.000000,Sensitivity=1.000000,Exponent=1.000000,bInvert=False))
-+AxisConfig=(AxisKeyName="MotionController_Left_Thumbstick_Z",AxisProperties=(DeadZone=0.000000,Sensitivity=1.000000,Exponent=1.000000,bInvert=False))
-+AxisConfig=(AxisKeyName="MotionController_Right_Thumbstick_Z",AxisProperties=(DeadZone=0.000000,Sensitivity=1.000000,Exponent=1.000000,bInvert=False))
-+AxisConfig=(AxisKeyName="OculusTouch_Left_Thumbstick",AxisProperties=(DeadZone=0.000000,Sensitivity=1.000000,Exponent=1.000000,bInvert=False))
-+AxisConfig=(AxisKeyName="OculusTouch_Left_FaceButton1",AxisProperties=(DeadZone=0.000000,Sensitivity=1.000000,Exponent=1.000000,bInvert=False))
-+AxisConfig=(AxisKeyName="OculusTouch_Left_Trigger",AxisProperties=(DeadZone=0.000000,Sensitivity=1.000000,Exponent=1.000000,bInvert=False))
-+AxisConfig=(AxisKeyName="OculusTouch_Left_FaceButton2",AxisProperties=(DeadZone=0.000000,Sensitivity=1.000000,Exponent=1.000000,bInvert=False))
-+AxisConfig=(AxisKeyName="OculusTouch_Left_IndexPointing",AxisProperties=(DeadZone=0.000000,Sensitivity=1.000000,Exponent=1.000000,bInvert=False))
-+AxisConfig=(AxisKeyName="OculusTouch_Left_ThumbUp",AxisProperties=(DeadZone=0.000000,Sensitivity=1.000000,Exponent=1.000000,bInvert=False))
-+AxisConfig=(AxisKeyName="OculusTouch_Right_Thumbstick",AxisProperties=(DeadZone=0.000000,Sensitivity=1.000000,Exponent=1.000000,bInvert=False))
-+AxisConfig=(AxisKeyName="OculusTouch_Right_FaceButton1",AxisProperties=(DeadZone=0.000000,Sensitivity=1.000000,Exponent=1.000000,bInvert=False))
-+AxisConfig=(AxisKeyName="OculusTouch_Right_Trigger",AxisProperties=(DeadZone=0.000000,Sensitivity=1.000000,Exponent=1.000000,bInvert=False))
-+AxisConfig=(AxisKeyName="OculusTouch_Right_FaceButton2",AxisProperties=(DeadZone=0.000000,Sensitivity=1.000000,Exponent=1.000000,bInvert=False))
-+AxisConfig=(AxisKeyName="OculusTouch_Right_IndexPointing",AxisProperties=(DeadZone=0.000000,Sensitivity=1.000000,Exponent=1.000000,bInvert=False))
-+AxisConfig=(AxisKeyName="OculusTouch_Right_ThumbUp",AxisProperties=(DeadZone=0.000000,Sensitivity=1.000000,Exponent=1.000000,bInvert=False))
-+AxisConfig=(AxisKeyName="OculusTouchpad_Touchpad_X",AxisProperties=(DeadZone=0.000000,Sensitivity=1.000000,Exponent=1.000000,bInvert=False))
-+AxisConfig=(AxisKeyName="OculusTouchpad_Touchpad_Y",AxisProperties=(DeadZone=0.000000,Sensitivity=1.000000,Exponent=1.000000,bInvert=False))
-+AxisConfig=(AxisKeyName="SteamVR_Knuckles_Left_HandGrip",AxisProperties=(DeadZone=0.000000,Sensitivity=1.000000,Exponent=1.000000,bInvert=False))
-+AxisConfig=(AxisKeyName="SteamVR_Knuckles_Left_IndexGrip",AxisProperties=(DeadZone=0.000000,Sensitivity=1.000000,Exponent=1.000000,bInvert=False))
-+AxisConfig=(AxisKeyName="SteamVR_Knuckles_Left_MiddleGrip",AxisProperties=(DeadZone=0.000000,Sensitivity=1.000000,Exponent=1.000000,bInvert=False))
-+AxisConfig=(AxisKeyName="SteamVR_Knuckles_Left_RingGrip",AxisProperties=(DeadZone=0.000000,Sensitivity=1.000000,Exponent=1.000000,bInvert=False))
-+AxisConfig=(AxisKeyName="SteamVR_Knuckles_Left_PinkyGrip",AxisProperties=(DeadZone=0.000000,Sensitivity=1.000000,Exponent=1.000000,bInvert=False))
-+AxisConfig=(AxisKeyName="SteamVR_Knuckles_Right_HandGrip",AxisProperties=(DeadZone=0.000000,Sensitivity=1.000000,Exponent=1.000000,bInvert=False))
-+AxisConfig=(AxisKeyName="SteamVR_Knuckles_Right_IndexGrip",AxisProperties=(DeadZone=0.000000,Sensitivity=1.000000,Exponent=1.000000,bInvert=False))
-+AxisConfig=(AxisKeyName="SteamVR_Knuckles_Right_MiddleGrip",AxisProperties=(DeadZone=0.000000,Sensitivity=1.000000,Exponent=1.000000,bInvert=False))
-+AxisConfig=(AxisKeyName="SteamVR_Knuckles_Right_RingGrip",AxisProperties=(DeadZone=0.000000,Sensitivity=1.000000,Exponent=1.000000,bInvert=False))
-+AxisConfig=(AxisKeyName="SteamVR_Knuckles_Right_PinkyGrip",AxisProperties=(DeadZone=0.000000,Sensitivity=1.000000,Exponent=1.000000,bInvert=False))
-+AxisConfig=(AxisKeyName="Vive_Left_Trigger_Axis",AxisProperties=(DeadZone=0.000000,Sensitivity=1.000000,Exponent=1.000000,bInvert=False))
-+AxisConfig=(AxisKeyName="Vive_Left_Trackpad_X",AxisProperties=(DeadZone=0.000000,Sensitivity=1.000000,Exponent=1.000000,bInvert=False))
-+AxisConfig=(AxisKeyName="Vive_Left_Trackpad_Y",AxisProperties=(DeadZone=0.000000,Sensitivity=1.000000,Exponent=1.000000,bInvert=False))
-+AxisConfig=(AxisKeyName="Vive_Right_Trigger_Axis",AxisProperties=(DeadZone=0.000000,Sensitivity=1.000000,Exponent=1.000000,bInvert=False))
-+AxisConfig=(AxisKeyName="Vive_Right_Trackpad_X",AxisProperties=(DeadZone=0.000000,Sensitivity=1.000000,Exponent=1.000000,bInvert=False))
-+AxisConfig=(AxisKeyName="Vive_Right_Trackpad_Y",AxisProperties=(DeadZone=0.000000,Sensitivity=1.000000,Exponent=1.000000,bInvert=False))
-+AxisConfig=(AxisKeyName="MixedReality_Left_Trigger_Axis",AxisProperties=(DeadZone=0.000000,Sensitivity=1.000000,Exponent=1.000000,bInvert=False))
-+AxisConfig=(AxisKeyName="MixedReality_Left_Thumbstick_X",AxisProperties=(DeadZone=0.000000,Sensitivity=1.000000,Exponent=1.000000,bInvert=False))
-+AxisConfig=(AxisKeyName="MixedReality_Left_Thumbstick_Y",AxisProperties=(DeadZone=0.000000,Sensitivity=1.000000,Exponent=1.000000,bInvert=False))
-+AxisConfig=(AxisKeyName="MixedReality_Left_Trackpad_X",AxisProperties=(DeadZone=0.000000,Sensitivity=1.000000,Exponent=1.000000,bInvert=False))
-+AxisConfig=(AxisKeyName="MixedReality_Left_Trackpad_Y",AxisProperties=(DeadZone=0.000000,Sensitivity=1.000000,Exponent=1.000000,bInvert=False))
-+AxisConfig=(AxisKeyName="MixedReality_Right_Trigger_Axis",AxisProperties=(DeadZone=0.000000,Sensitivity=1.000000,Exponent=1.000000,bInvert=False))
-+AxisConfig=(AxisKeyName="MixedReality_Right_Thumbstick_X",AxisProperties=(DeadZone=0.000000,Sensitivity=1.000000,Exponent=1.000000,bInvert=False))
-+AxisConfig=(AxisKeyName="MixedReality_Right_Thumbstick_Y",AxisProperties=(DeadZone=0.000000,Sensitivity=1.000000,Exponent=1.000000,bInvert=False))
-+AxisConfig=(AxisKeyName="MixedReality_Right_Trackpad_X",AxisProperties=(DeadZone=0.000000,Sensitivity=1.000000,Exponent=1.000000,bInvert=False))
-+AxisConfig=(AxisKeyName="MixedReality_Right_Trackpad_Y",AxisProperties=(DeadZone=0.000000,Sensitivity=1.000000,Exponent=1.000000,bInvert=False))
-+AxisConfig=(AxisKeyName="OculusGo_Left_Trackpad_X",AxisProperties=(DeadZone=0.000000,Sensitivity=1.000000,Exponent=1.000000,bInvert=False))
-+AxisConfig=(AxisKeyName="OculusGo_Left_Trackpad_Y",AxisProperties=(DeadZone=0.000000,Sensitivity=1.000000,Exponent=1.000000,bInvert=False))
-+AxisConfig=(AxisKeyName="OculusGo_Right_Trackpad_X",AxisProperties=(DeadZone=0.000000,Sensitivity=1.000000,Exponent=1.000000,bInvert=False))
-+AxisConfig=(AxisKeyName="OculusGo_Right_Trackpad_Y",AxisProperties=(DeadZone=0.000000,Sensitivity=1.000000,Exponent=1.000000,bInvert=False))
-+AxisConfig=(AxisKeyName="OculusTouch_Left_Grip_Axis",AxisProperties=(DeadZone=0.000000,Sensitivity=1.000000,Exponent=1.000000,bInvert=False))
-+AxisConfig=(AxisKeyName="OculusTouch_Left_Trigger_Axis",AxisProperties=(DeadZone=0.000000,Sensitivity=1.000000,Exponent=1.000000,bInvert=False))
-+AxisConfig=(AxisKeyName="OculusTouch_Left_Thumbstick_X",AxisProperties=(DeadZone=0.000000,Sensitivity=1.000000,Exponent=1.000000,bInvert=False))
-+AxisConfig=(AxisKeyName="OculusTouch_Left_Thumbstick_Y",AxisProperties=(DeadZone=0.000000,Sensitivity=1.000000,Exponent=1.000000,bInvert=False))
-+AxisConfig=(AxisKeyName="OculusTouch_Right_Grip_Axis",AxisProperties=(DeadZone=0.000000,Sensitivity=1.000000,Exponent=1.000000,bInvert=False))
-+AxisConfig=(AxisKeyName="OculusTouch_Right_Trigger_Axis",AxisProperties=(DeadZone=0.000000,Sensitivity=1.000000,Exponent=1.000000,bInvert=False))
-+AxisConfig=(AxisKeyName="OculusTouch_Right_Thumbstick_X",AxisProperties=(DeadZone=0.000000,Sensitivity=1.000000,Exponent=1.000000,bInvert=False))
-+AxisConfig=(AxisKeyName="OculusTouch_Right_Thumbstick_Y",AxisProperties=(DeadZone=0.000000,Sensitivity=1.000000,Exponent=1.000000,bInvert=False))
-+AxisConfig=(AxisKeyName="ValveIndex_Left_Grip_Axis",AxisProperties=(DeadZone=0.000000,Sensitivity=1.000000,Exponent=1.000000,bInvert=False))
-+AxisConfig=(AxisKeyName="ValveIndex_Left_Grip_Force",AxisProperties=(DeadZone=0.000000,Sensitivity=1.000000,Exponent=1.000000,bInvert=False))
-+AxisConfig=(AxisKeyName="ValveIndex_Left_Trigger_Axis",AxisProperties=(DeadZone=0.000000,Sensitivity=1.000000,Exponent=1.000000,bInvert=False))
-+AxisConfig=(AxisKeyName="ValveIndex_Left_Thumbstick_X",AxisProperties=(DeadZone=0.000000,Sensitivity=1.000000,Exponent=1.000000,bInvert=False))
-+AxisConfig=(AxisKeyName="ValveIndex_Left_Thumbstick_Y",AxisProperties=(DeadZone=0.000000,Sensitivity=1.000000,Exponent=1.000000,bInvert=False))
-+AxisConfig=(AxisKeyName="ValveIndex_Left_Trackpad_X",AxisProperties=(DeadZone=0.000000,Sensitivity=1.000000,Exponent=1.000000,bInvert=False))
-+AxisConfig=(AxisKeyName="ValveIndex_Left_Trackpad_Y",AxisProperties=(DeadZone=0.000000,Sensitivity=1.000000,Exponent=1.000000,bInvert=False))
-+AxisConfig=(AxisKeyName="ValveIndex_Left_Trackpad_Force",AxisProperties=(DeadZone=0.000000,Sensitivity=1.000000,Exponent=1.000000,bInvert=False))
-+AxisConfig=(AxisKeyName="ValveIndex_Left_Trackpad_Touch",AxisProperties=(DeadZone=0.000000,Sensitivity=1.000000,Exponent=1.000000,bInvert=False))
-+AxisConfig=(AxisKeyName="ValveIndex_Right_Grip_Axis",AxisProperties=(DeadZone=0.000000,Sensitivity=1.000000,Exponent=1.000000,bInvert=False))
-+AxisConfig=(AxisKeyName="ValveIndex_Right_Grip_Force",AxisProperties=(DeadZone=0.000000,Sensitivity=1.000000,Exponent=1.000000,bInvert=False))
-+AxisConfig=(AxisKeyName="ValveIndex_Right_Trigger_Axis",AxisProperties=(DeadZone=0.000000,Sensitivity=1.000000,Exponent=1.000000,bInvert=False))
-+AxisConfig=(AxisKeyName="ValveIndex_Right_Thumbstick_X",AxisProperties=(DeadZone=0.000000,Sensitivity=1.000000,Exponent=1.000000,bInvert=False))
-+AxisConfig=(AxisKeyName="ValveIndex_Right_Thumbstick_Y",AxisProperties=(DeadZone=0.000000,Sensitivity=1.000000,Exponent=1.000000,bInvert=False))
-+AxisConfig=(AxisKeyName="ValveIndex_Right_Trackpad_X",AxisProperties=(DeadZone=0.000000,Sensitivity=1.000000,Exponent=1.000000,bInvert=False))
-+AxisConfig=(AxisKeyName="ValveIndex_Right_Trackpad_Y",AxisProperties=(DeadZone=0.000000,Sensitivity=1.000000,Exponent=1.000000,bInvert=False))
-+AxisConfig=(AxisKeyName="ValveIndex_Right_Trackpad_Force",AxisProperties=(DeadZone=0.000000,Sensitivity=1.000000,Exponent=1.000000,bInvert=False))
 bAltEnterTogglesFullscreen=True
 bF11TogglesFullscreen=True
 bUseMouseForTouch=False
@@ -108,49 +27,16 @@ FOVScale=0.011110
 DoubleClickTime=0.200000
 +ActionMappings=(ActionName="Jump",bShift=False,bCtrl=False,bAlt=False,bCmd=False,Key=SpaceBar)
 +ActionMappings=(ActionName="Jump",bShift=False,bCtrl=False,bAlt=False,bCmd=False,Key=Gamepad_FaceButton_Bottom)
-+ActionMappings=(ActionName="ResetVR",bShift=False,bCtrl=False,bAlt=False,bCmd=False,Key=R)
-+ActionMappings=(ActionName="Jump",bShift=False,bCtrl=False,bAlt=False,bCmd=False,Key=Vive_Left_Trigger_Click)
-+ActionMappings=(ActionName="Jump",bShift=False,bCtrl=False,bAlt=False,bCmd=False,Key=Vive_Right_Trigger_Click)
-+ActionMappings=(ActionName="Jump",bShift=False,bCtrl=False,bAlt=False,bCmd=False,Key=MixedReality_Left_Trigger_Click)
-+ActionMappings=(ActionName="Jump",bShift=False,bCtrl=False,bAlt=False,bCmd=False,Key=MixedReality_Right_Trigger_Click)
-+ActionMappings=(ActionName="Jump",bShift=False,bCtrl=False,bAlt=False,bCmd=False,Key=OculusGo_Left_Trigger_Click)
-+ActionMappings=(ActionName="Jump",bShift=False,bCtrl=False,bAlt=False,bCmd=False,Key=OculusTouch_Left_Trigger_Click)
-+ActionMappings=(ActionName="Jump",bShift=False,bCtrl=False,bAlt=False,bCmd=False,Key=OculusTouch_Right_Trigger_Click)
-+ActionMappings=(ActionName="Jump",bShift=False,bCtrl=False,bAlt=False,bCmd=False,Key=ValveIndex_Left_Trigger_Click)
-+ActionMappings=(ActionName="Jump",bShift=False,bCtrl=False,bAlt=False,bCmd=False,Key=ValveIndex_Right_Trigger_Click)
-+ActionMappings=(ActionName="ResetVR",bShift=False,bCtrl=False,bAlt=False,bCmd=False,Key=Vive_Left_Grip_Click)
-+ActionMappings=(ActionName="ResetVR",bShift=False,bCtrl=False,bAlt=False,bCmd=False,Key=MixedReality_Left_Thumbstick_Click)
-+ActionMappings=(ActionName="ResetVR",bShift=False,bCtrl=False,bAlt=False,bCmd=False,Key=OculusGo_Left_Trackpad_Click)
-+ActionMappings=(ActionName="ResetVR",bShift=False,bCtrl=False,bAlt=False,bCmd=False,Key=OculusTouch_Left_Thumbstick_Click)
-+ActionMappings=(ActionName="ResetVR",bShift=False,bCtrl=False,bAlt=False,bCmd=False,Key=ValveIndex_Left_Thumbstick_Click)
-+AxisMappings=(AxisName="MoveForward",Scale=1.000000,Key=W)
-+AxisMappings=(AxisName="MoveForward",Scale=-1.000000,Key=S)
-+AxisMappings=(AxisName="MoveForward",Scale=1.000000,Key=Up)
-+AxisMappings=(AxisName="MoveForward",Scale=-1.000000,Key=Down)
-+AxisMappings=(AxisName="MoveForward",Scale=1.000000,Key=Gamepad_LeftY)
-+AxisMappings=(AxisName="MoveRight",Scale=-1.000000,Key=A)
-+AxisMappings=(AxisName="MoveRight",Scale=1.000000,Key=D)
-+AxisMappings=(AxisName="MoveRight",Scale=1.000000,Key=Gamepad_LeftX)
-+AxisMappings=(AxisName="TurnRate",Scale=1.000000,Key=Gamepad_RightX)
-+AxisMappings=(AxisName="TurnRate",Scale=-1.000000,Key=Left)
-+AxisMappings=(AxisName="TurnRate",Scale=1.000000,Key=Right)
-+AxisMappings=(AxisName="Turn",Scale=1.000000,Key=MouseX)
-+AxisMappings=(AxisName="LookUpRate",Scale=1.000000,Key=Gamepad_RightY)
-+AxisMappings=(AxisName="LookUp",Scale=-1.000000,Key=MouseY)
-+AxisMappings=(AxisName="TurnRate",Scale=-1.000000,Key=Vive_Right_Trackpad_X)
-+AxisMappings=(AxisName="MoveForward",Scale=1.000000,Key=Vive_Left_Trackpad_Y)
-+AxisMappings=(AxisName="MoveRight",Scale=1.000000,Key=Vive_Left_Trackpad_X)
-+AxisMappings=(AxisName="MoveRight",Scale=1.000000,Key=MixedReality_Left_Thumbstick_X)
-+AxisMappings=(AxisName="MoveRight",Scale=1.000000,Key=OculusGo_Left_Trackpad_X)
-+AxisMappings=(AxisName="MoveForward",Scale=1.000000,Key=MixedReality_Left_Thumbstick_Y)
-+AxisMappings=(AxisName="MoveForward",Scale=1.000000,Key=OculusGo_Left_Trackpad_Y)
-+AxisMappings=(AxisName="TurnRate",Scale=-1.000000,Key=MixedReality_Right_Thumbstick_X)
-+AxisMappings=(AxisName="TurnRate",Scale=-1.000000,Key=OculusTouch_Right_Thumbstick_X)
-+AxisMappings=(AxisName="TurnRate",Scale=-1.000000,Key=ValveIndex_Right_Thumbstick_X)
-+AxisMappings=(AxisName="MoveForward",Scale=1.000000,Key=OculusTouch_Left_Thumbstick_Y)
-+AxisMappings=(AxisName="MoveForward",Scale=1.000000,Key=ValveIndex_Left_Thumbstick_Y)
-+AxisMappings=(AxisName="MoveRight",Scale=1.000000,Key=OculusTouch_Left_Thumbstick_X)
-+AxisMappings=(AxisName="MoveRight",Scale=1.000000,Key=ValveIndex_Left_Thumbstick_X)
++AxisMappings=(AxisName="Move Forward / Backward",Scale=1.000000,Key=W)
++AxisMappings=(AxisName="Move Forward / Backward",Scale=-1.000000,Key=S)
++AxisMappings=(AxisName="Move Forward / Backward",Scale=1.000000,Key=Gamepad_LeftY)
++AxisMappings=(AxisName="Move Right / Left",Scale=-1.000000,Key=A)
++AxisMappings=(AxisName="Move Right / Left",Scale=1.000000,Key=D)
++AxisMappings=(AxisName="Move Right / Left",Scale=1.000000,Key=Gamepad_LeftX)
++AxisMappings=(AxisName="Turn Right / Left Gamepad",Scale=1.000000,Key=Gamepad_RightX)
++AxisMappings=(AxisName="Turn Right / Left Mouse",Scale=1.000000,Key=MouseX)
++AxisMappings=(AxisName="Look Up / Down Gamepad",Scale=1.000000,Key=Gamepad_RightY)
++AxisMappings=(AxisName="Look Up / Down Mouse",Scale=-1.000000,Key=MouseY)
 DefaultTouchInterface=/Engine/MobileResources/HUD/DefaultVirtualJoysticks.DefaultVirtualJoysticks
 +ConsoleKeys=Tilde
 
diff --git a/Templates/TP_ThirdPerson/Config/TemplateDefs.ini b/Templates/TP_ThirdPerson/Config/TemplateDefs.ini
--- a/Templates/TP_ThirdPerson/Config/TemplateDefs.ini
+++ b/Templates/TP_ThirdPerson/Config/TemplateDefs.ini
@@ -48,6 +48,6 @@ FilenameReplacements=(Extensions=("cpp","h","ini","cs"),From="%TEMPLATENAME%",To
 ReplacementsInFiles=(Extensions=("cpp","h","ini","cs"),From="%TEMPLATENAME_UPPERCASE%",To="%PROJECTNAME_UPPERCASE%",bCaseSensitive=true)
 ReplacementsInFiles=(Extensions=("cpp","h","ini","cs"),From="%TEMPLATENAME_LOWERCASE%",To="%PROJECTNAME_LOWERCASE%",bCaseSensitive=true)
 ReplacementsInFiles=(Extensions=("cpp","h","ini","cs"),From="%TEMPLATENAME%",To="%PROJECTNAME%",bCaseSensitive=false)
-SharedContentPacks=(MountName="ThirdPerson",DetailLevels=("Standard"))
-SharedContentPacks=(MountName="Geometry",DetailLevels=("Standard","High"))
-SharedContentPacks=(MountName="Mannequin",DetailLevels=("Standard","High"))
+SharedContentPacks=(MountName="LevelPrototyping",DetailLevels=("High"))
+SharedContentPacks=(MountName="Mannequin",DetailLevels=("High"))
+EditDetailLevelPreference="High"
\ No newline at end of file
diff --git a/Templates/TP_ThirdPerson/Config/config.ini b/Templates/TP_ThirdPerson/Config/config.ini
--- a/Templates/TP_ThirdPerson/Config/config.ini
+++ b/Templates/TP_ThirdPerson/Config/config.ini
@@ -1,24 +1,16 @@
 [InputSettings]
-+ActionMappings=(ActionName="Jump", Key=SpaceBar)
-+ActionMappings=(ActionName="Jump", Key=Gamepad_FaceButton_Bottom)
-
-+AxisMappings=(AxisName="MoveForward", Key=W, Scale=1.f)
-+AxisMappings=(AxisName="MoveForward", Key=S, Scale=-1.f)
-+AxisMappings=(AxisName="MoveForward", Key=Up, Scale=1.f)
-+AxisMappings=(AxisName="MoveForward", Key=Down, Scale=-1.f)
-+AxisMappings=(AxisName="MoveForward", Key=Gamepad_LeftY, Scale=1.f)
-
-+AxisMappings=(AxisName="MoveRight", Key=A, Scale=-1.f)
-+AxisMappings=(AxisName="MoveRight", Key=D, Scale=1.f)
-+AxisMappings=(AxisName="MoveRight", Key=Gamepad_LeftX, Scale=1.f)
-
-+AxisMappings=(AxisName="TurnRate", Key=Gamepad_RightX, Scale=1.f)
-+AxisMappings=(AxisName="TurnRate", Key=Left, Scale=-1.f)
-+AxisMappings=(AxisName="TurnRate", Key=Right, Scale=1.f)
-+AxisMappings=(AxisName="Turn", Key=MouseX, Scale=1.f)
-
-+AxisMappings=(AxisName="LookUpRate", Key=Gamepad_RightY, Scale=1.f)
-+AxisMappings=(AxisName="LookUp", Key=MouseY, Scale=-1.f)
++ActionMappings=(ActionName="Jump",bShift=False,bCtrl=False,bAlt=False,bCmd=False,Key=SpaceBar)
++ActionMappings=(ActionName="Jump",bShift=False,bCtrl=False,bAlt=False,bCmd=False,Key=Gamepad_FaceButton_Bottom)
++AxisMappings=(AxisName="Move Forward / Backward",Scale=1.000000,Key=W)
++AxisMappings=(AxisName="Move Forward / Backward",Scale=-1.000000,Key=S)
++AxisMappings=(AxisName="Move Forward / Backward",Scale=1.000000,Key=Gamepad_LeftY)
++AxisMappings=(AxisName="Move Right / Left",Scale=-1.000000,Key=A)
++AxisMappings=(AxisName="Move Right / Left",Scale=1.000000,Key=D)
++AxisMappings=(AxisName="Move Right / Left",Scale=1.000000,Key=Gamepad_LeftX)
++AxisMappings=(AxisName="Turn Right / Left Gamepad",Scale=1.000000,Key=Gamepad_RightX)
++AxisMappings=(AxisName="Turn Right / Left Mouse",Scale=1.000000,Key=MouseX)
++AxisMappings=(AxisName="Look Up / Down Gamepad",Scale=1.000000,Key=Gamepad_RightY)
++AxisMappings=(AxisName="Look Up / Down Mouse",Scale=-1.000000,Key=MouseY)
 
 [Redirects]
 GameName=TP_ThirdPerson
diff --git a/Templates/TP_ThirdPerson/Source/TP_ThirdPerson/TP_ThirdPersonCharacter.cpp b/Templates/TP_ThirdPerson/Source/TP_ThirdPerson/TP_ThirdPersonCharacter.cpp
--- a/Templates/TP_ThirdPerson/Source/TP_ThirdPerson/TP_ThirdPersonCharacter.cpp
+++ b/Templates/TP_ThirdPerson/Source/TP_ThirdPerson/TP_ThirdPersonCharacter.cpp
@@ -17,8 +17,8 @@ ATP_ThirdPersonCharacter::ATP_ThirdPersonCharacter()
 	GetCapsuleComponent()->InitCapsuleSize(42.f, 96.0f);
 
 	// set our turn rates for input
-	BaseTurnRate = 45.f;
-	BaseLookUpRate = 45.f;
+	BaseTurnRate = 50.f;
+	BaseLookUpRate = 50.f;
 
 	// Don't rotate when the controller rotates. Let that just affect the camera.
 	bUseControllerRotationPitch = false;
@@ -29,12 +29,12 @@ ATP_ThirdPersonCharacter::ATP_ThirdPersonCharacter()
 	GetCharacterMovement()->bOrientRotationToMovement = true; // Character moves in the direction of input...	
 	GetCharacterMovement()->RotationRate = FRotator(0.0f, 540.0f, 0.0f); // ...at this rotation rate
 	GetCharacterMovement()->JumpZVelocity = 600.f;
-	GetCharacterMovement()->AirControl = 0.2f;
+	GetCharacterMovement()->AirControl = 0.35f;
 
 	// Create a camera boom (pulls in towards the player if there is a collision)
 	CameraBoom = CreateDefaultSubobject<USpringArmComponent>(TEXT("CameraBoom"));
 	CameraBoom->SetupAttachment(RootComponent);
-	CameraBoom->TargetArmLength = 300.0f; // The camera follows at this distance behind the character	
+	CameraBoom->TargetArmLength = 400.0f; // The camera follows at this distance behind the character	
 	CameraBoom->bUsePawnControlRotation = true; // Rotate the arm based on the controller
 
 	// Create a follow camera
@@ -56,16 +56,16 @@ void ATP_ThirdPersonCharacter::SetupPlayerInputComponent(class UInputComponent*
 	PlayerInputComponent->BindAction("Jump", IE_Pressed, this, &ACharacter::Jump);
 	PlayerInputComponent->BindAction("Jump", IE_Released, this, &ACharacter::StopJumping);
 
-	PlayerInputComponent->BindAxis("MoveForward", this, &ATP_ThirdPersonCharacter::MoveForward);
-	PlayerInputComponent->BindAxis("MoveRight", this, &ATP_ThirdPersonCharacter::MoveRight);
+	PlayerInputComponent->BindAxis("Move Forward / Backward", this, &ATP_ThirdPersonCharacter::MoveForward);
+	PlayerInputComponent->BindAxis("Move Right / Left", this, &ATP_ThirdPersonCharacter::MoveRight);
 
 	// We have 2 versions of the rotation bindings to handle different kinds of devices differently
 	// "turn" handles devices that provide an absolute delta, such as a mouse.
 	// "turnrate" is for devices that we choose to treat as a rate of change, such as an analog joystick
-	PlayerInputComponent->BindAxis("Turn", this, &APawn::AddControllerYawInput);
-	PlayerInputComponent->BindAxis("TurnRate", this, &ATP_ThirdPersonCharacter::TurnAtRate);
-	PlayerInputComponent->BindAxis("LookUp", this, &APawn::AddControllerPitchInput);
-	PlayerInputComponent->BindAxis("LookUpRate", this, &ATP_ThirdPersonCharacter::LookUpAtRate);
+	PlayerInputComponent->BindAxis("Turn Right / Left Mouse", this, &APawn::AddControllerYawInput);
+	PlayerInputComponent->BindAxis("Turn Right / Left Gamepad", this, &ATP_ThirdPersonCharacter::TurnAtRate);
+	PlayerInputComponent->BindAxis("Look Up / Down Mouse", this, &APawn::AddControllerPitchInput);
+	PlayerInputComponent->BindAxis("Look Up / Down Gamepad", this, &ATP_ThirdPersonCharacter::LookUpAtRate);
 
 	// handle touch devices
 	PlayerInputComponent->BindTouch(IE_Pressed, this, &ATP_ThirdPersonCharacter::TouchStarted);
diff --git a/Templates/TP_ThirdPerson/Source/TP_ThirdPerson/TP_ThirdPersonGameMode.cpp b/Templates/TP_ThirdPerson/Source/TP_ThirdPerson/TP_ThirdPersonGameMode.cpp
--- a/Templates/TP_ThirdPerson/Source/TP_ThirdPerson/TP_ThirdPersonGameMode.cpp
+++ b/Templates/TP_ThirdPerson/Source/TP_ThirdPerson/TP_ThirdPersonGameMode.cpp
@@ -7,7 +7,7 @@
 ATP_ThirdPersonGameMode::ATP_ThirdPersonGameMode()
 {
 	// set default pawn class to our Blueprinted character
-	static ConstructorHelpers::FClassFinder<APawn> PlayerPawnBPClass(TEXT("/Game/ThirdPersonCPP/Blueprints/ThirdPersonCharacter"));
+	static ConstructorHelpers::FClassFinder<APawn> PlayerPawnBPClass(TEXT("/Game/ThirdPerson/Blueprints/ThirdPersonCharacter"));
 	if (PlayerPawnBPClass.Class != NULL)
 	{
 		DefaultPawnClass = PlayerPawnBPClass.Class;
diff --git a/Templates/TP_ThirdPerson/TP_ThirdPerson.uproject b/Templates/TP_ThirdPerson/TP_ThirdPerson.uproject
--- a/Templates/TP_ThirdPerson/TP_ThirdPerson.uproject
+++ b/Templates/TP_ThirdPerson/TP_ThirdPerson.uproject
@@ -17,6 +17,52 @@
 			"TargetAllowList": [
 				"Editor"
 			]
+		},
+		{
+			"Name": "OculusVR",
+			"Enabled": false,
+			"SupportedTargetPlatforms": [
+				"Win64",
+				"Android"
+			]
+		},
+		{
+			"Name": "OpenXREyeTracker",
+			"Enabled": false,
+			"SupportedTargetPlatforms": [
+				"Win64",
+				"Linux",
+				"HoloLens",
+				"Android"
+			]
+		},
+		{
+			"Name": "OpenXRHandTracking",
+			"Enabled": false,
+			"SupportedTargetPlatforms": [
+				"Win64",
+				"Linux",
+				"HoloLens",
+				"Android"
+			]
+		},
+		{
+			"Name": "OpenXR",
+			"Enabled": false,
+			"SupportedTargetPlatforms": [
+				"Win64",
+				"Linux",
+				"Android",
+				"HoloLens"
+			]
+		},
+		{
+			"Name": "SteamVR",
+			"Enabled": false,
+			"SupportedTargetPlatforms": [
+				"Win64",
+				"Linux"
+			]
 		}
-	]	
+	]
 }
\ No newline at end of file
diff --git a/Templates/TP_ThirdPerson/manifest.json b/Templates/TP_ThirdPerson/manifest.json
--- a/Templates/TP_ThirdPerson/manifest.json
+++ b/Templates/TP_ThirdPerson/manifest.json
@@ -54,34 +54,25 @@
 	],
 	"ClassTypes": "GameMode, Character, SpringArmComponent, CameraComponent, InputComponent",
 	"Category": "CodeFeature",
-	"FocusAsset": "/Game/ThirdPersonCPP/ThirdPersonOverview.ThirdPersonOverview",
+	"FocusAsset": "/Game/ThirdPerson/Blueprints",
 	"Thumbnail": "TP_ThirdPerson.png",
 	"Screenshots":
 	[
 		"TP_ThirdPerson_Preview.png"
 	],
 	"AdditionalFeaturePacks":
-	[
-		{
-			"MountName": "ThirdPerson",
-			"DetailLevels":
-			[
-				"Standard"
-			]
-		},
+	[		
 		{
-			"MountName": "Geometry",
+			"MountName": "LevelPrototyping",
 			"DetailLevels":
-			[
-				"Standard",
+			[			
 				"High"
 			]
 		},
 		{
 			"MountName": "Mannequin",
 			"DetailLevels":
-			[
-				"Standard",
+			[		
 				"High"
 			]
 		}
diff --git a/Templates/TP_ThirdPersonBP/Config/DefaultEngine.ini b/Templates/TP_ThirdPersonBP/Config/DefaultEngine.ini
--- a/Templates/TP_ThirdPersonBP/Config/DefaultEngine.ini
+++ b/Templates/TP_ThirdPersonBP/Config/DefaultEngine.ini
@@ -2,14 +2,21 @@
 GameName=TP_ThirdPersonBP
 
 [/Script/EngineSettings.GameMapsSettings]
-EditorStartupMap=/Game/ThirdPersonBP/Maps/ThirdPersonExampleMap
-GameDefaultMap=/Game/ThirdPersonBP/Maps/ThirdPersonExampleMap
+EditorStartupMap=/Game/ThirdPerson/Maps/ThirdPersonMap.ThirdPersonMap
+GameDefaultMap=/Game/ThirdPerson/Maps/ThirdPersonMap.ThirdPersonMap
 TransitionMap=
 bUseSplitscreen=True
 TwoPlayerSplitscreenLayout=Horizontal
 ThreePlayerSplitscreenLayout=FavorTop
-GlobalDefaultGameMode=/Game/ThirdPersonBP/Blueprints/ThirdPersonGameMode.ThirdPersonGameMode_C
+GlobalDefaultGameMode=/Game/ThirdPerson/Blueprints/BP_ThirdPersonGameMode.BP_ThirdPersonGameMode_C
 GlobalDefaultServerGameMode=None
 
 [/Script/IOSRuntimeSettings.IOSRuntimeSettings]
 MinimumiOSVersion=IOS_13
+
+[/Script/Engine.RendererSettings]
+r.ReflectionMethod=1
+r.GenerateMeshDistanceFields=True
+r.DynamicGlobalIlluminationMethod=1
+r.Lumen.TraceMeshSDFs=0
+
diff --git a/Templates/TP_ThirdPersonBP/Config/DefaultInput.ini b/Templates/TP_ThirdPersonBP/Config/DefaultInput.ini
--- a/Templates/TP_ThirdPersonBP/Config/DefaultInput.ini
+++ b/Templates/TP_ThirdPersonBP/Config/DefaultInput.ini
@@ -1,4 +1,5 @@
 
+
 [/Script/Engine.InputSettings]
 +AxisConfig=(AxisKeyName="Gamepad_LeftX",AxisProperties=(DeadZone=0.250000,Sensitivity=1.000000,Exponent=1.000000,bInvert=False))
 +AxisConfig=(AxisKeyName="Gamepad_LeftY",AxisProperties=(DeadZone=0.250000,Sensitivity=1.000000,Exponent=1.000000,bInvert=False))
@@ -9,95 +10,15 @@
 +AxisConfig=(AxisKeyName="MouseWheelAxis",AxisProperties=(DeadZone=0.000000,Sensitivity=1.000000,Exponent=1.000000,bInvert=False))
 +AxisConfig=(AxisKeyName="Gamepad_LeftTriggerAxis",AxisProperties=(DeadZone=0.000000,Sensitivity=1.000000,Exponent=1.000000,bInvert=False))
 +AxisConfig=(AxisKeyName="Gamepad_RightTriggerAxis",AxisProperties=(DeadZone=0.000000,Sensitivity=1.000000,Exponent=1.000000,bInvert=False))
-+AxisConfig=(AxisKeyName="MotionController_Left_Thumbstick_X",AxisProperties=(DeadZone=0.000000,Sensitivity=1.000000,Exponent=1.000000,bInvert=False))
-+AxisConfig=(AxisKeyName="MotionController_Left_Thumbstick_Y",AxisProperties=(DeadZone=0.000000,Sensitivity=1.000000,Exponent=1.000000,bInvert=False))
-+AxisConfig=(AxisKeyName="MotionController_Left_TriggerAxis",AxisProperties=(DeadZone=0.000000,Sensitivity=1.000000,Exponent=1.000000,bInvert=False))
-+AxisConfig=(AxisKeyName="MotionController_Left_Grip1Axis",AxisProperties=(DeadZone=0.000000,Sensitivity=1.000000,Exponent=1.000000,bInvert=False))
-+AxisConfig=(AxisKeyName="MotionController_Left_Grip2Axis",AxisProperties=(DeadZone=0.000000,Sensitivity=1.000000,Exponent=1.000000,bInvert=False))
-+AxisConfig=(AxisKeyName="MotionController_Right_Thumbstick_X",AxisProperties=(DeadZone=0.000000,Sensitivity=1.000000,Exponent=1.000000,bInvert=False))
-+AxisConfig=(AxisKeyName="MotionController_Right_Thumbstick_Y",AxisProperties=(DeadZone=0.000000,Sensitivity=1.000000,Exponent=1.000000,bInvert=False))
-+AxisConfig=(AxisKeyName="MotionController_Right_TriggerAxis",AxisProperties=(DeadZone=0.000000,Sensitivity=1.000000,Exponent=1.000000,bInvert=False))
-+AxisConfig=(AxisKeyName="MotionController_Right_Grip1Axis",AxisProperties=(DeadZone=0.000000,Sensitivity=1.000000,Exponent=1.000000,bInvert=False))
-+AxisConfig=(AxisKeyName="MotionController_Right_Grip2Axis",AxisProperties=(DeadZone=0.000000,Sensitivity=1.000000,Exponent=1.000000,bInvert=False))
 +AxisConfig=(AxisKeyName="Gamepad_Special_Left_X",AxisProperties=(DeadZone=0.000000,Sensitivity=1.000000,Exponent=1.000000,bInvert=False))
 +AxisConfig=(AxisKeyName="Gamepad_Special_Left_Y",AxisProperties=(DeadZone=0.000000,Sensitivity=1.000000,Exponent=1.000000,bInvert=False))
-+AxisConfig=(AxisKeyName="MotionController_Left_Thumbstick_Z",AxisProperties=(DeadZone=0.000000,Sensitivity=1.000000,Exponent=1.000000,bInvert=False))
-+AxisConfig=(AxisKeyName="MotionController_Right_Thumbstick_Z",AxisProperties=(DeadZone=0.000000,Sensitivity=1.000000,Exponent=1.000000,bInvert=False))
-+AxisConfig=(AxisKeyName="OculusTouch_Left_Thumbstick",AxisProperties=(DeadZone=0.000000,Sensitivity=1.000000,Exponent=1.000000,bInvert=False))
-+AxisConfig=(AxisKeyName="OculusTouch_Left_FaceButton1",AxisProperties=(DeadZone=0.000000,Sensitivity=1.000000,Exponent=1.000000,bInvert=False))
-+AxisConfig=(AxisKeyName="OculusTouch_Left_Trigger",AxisProperties=(DeadZone=0.000000,Sensitivity=1.000000,Exponent=1.000000,bInvert=False))
-+AxisConfig=(AxisKeyName="OculusTouch_Left_FaceButton2",AxisProperties=(DeadZone=0.000000,Sensitivity=1.000000,Exponent=1.000000,bInvert=False))
-+AxisConfig=(AxisKeyName="OculusTouch_Left_IndexPointing",AxisProperties=(DeadZone=0.000000,Sensitivity=1.000000,Exponent=1.000000,bInvert=False))
-+AxisConfig=(AxisKeyName="OculusTouch_Left_ThumbUp",AxisProperties=(DeadZone=0.000000,Sensitivity=1.000000,Exponent=1.000000,bInvert=False))
-+AxisConfig=(AxisKeyName="OculusTouch_Right_Thumbstick",AxisProperties=(DeadZone=0.000000,Sensitivity=1.000000,Exponent=1.000000,bInvert=False))
-+AxisConfig=(AxisKeyName="OculusTouch_Right_FaceButton1",AxisProperties=(DeadZone=0.000000,Sensitivity=1.000000,Exponent=1.000000,bInvert=False))
-+AxisConfig=(AxisKeyName="OculusTouch_Right_Trigger",AxisProperties=(DeadZone=0.000000,Sensitivity=1.000000,Exponent=1.000000,bInvert=False))
-+AxisConfig=(AxisKeyName="OculusTouch_Right_FaceButton2",AxisProperties=(DeadZone=0.000000,Sensitivity=1.000000,Exponent=1.000000,bInvert=False))
-+AxisConfig=(AxisKeyName="OculusTouch_Right_IndexPointing",AxisProperties=(DeadZone=0.000000,Sensitivity=1.000000,Exponent=1.000000,bInvert=False))
-+AxisConfig=(AxisKeyName="OculusTouch_Right_ThumbUp",AxisProperties=(DeadZone=0.000000,Sensitivity=1.000000,Exponent=1.000000,bInvert=False))
-+AxisConfig=(AxisKeyName="OculusTouchpad_Touchpad_X",AxisProperties=(DeadZone=0.000000,Sensitivity=1.000000,Exponent=1.000000,bInvert=False))
-+AxisConfig=(AxisKeyName="OculusTouchpad_Touchpad_Y",AxisProperties=(DeadZone=0.000000,Sensitivity=1.000000,Exponent=1.000000,bInvert=False))
-+AxisConfig=(AxisKeyName="SteamVR_Knuckles_Left_HandGrip",AxisProperties=(DeadZone=0.000000,Sensitivity=1.000000,Exponent=1.000000,bInvert=False))
-+AxisConfig=(AxisKeyName="SteamVR_Knuckles_Left_IndexGrip",AxisProperties=(DeadZone=0.000000,Sensitivity=1.000000,Exponent=1.000000,bInvert=False))
-+AxisConfig=(AxisKeyName="SteamVR_Knuckles_Left_MiddleGrip",AxisProperties=(DeadZone=0.000000,Sensitivity=1.000000,Exponent=1.000000,bInvert=False))
-+AxisConfig=(AxisKeyName="SteamVR_Knuckles_Left_RingGrip",AxisProperties=(DeadZone=0.000000,Sensitivity=1.000000,Exponent=1.000000,bInvert=False))
-+AxisConfig=(AxisKeyName="SteamVR_Knuckles_Left_PinkyGrip",AxisProperties=(DeadZone=0.000000,Sensitivity=1.000000,Exponent=1.000000,bInvert=False))
-+AxisConfig=(AxisKeyName="SteamVR_Knuckles_Right_HandGrip",AxisProperties=(DeadZone=0.000000,Sensitivity=1.000000,Exponent=1.000000,bInvert=False))
-+AxisConfig=(AxisKeyName="SteamVR_Knuckles_Right_IndexGrip",AxisProperties=(DeadZone=0.000000,Sensitivity=1.000000,Exponent=1.000000,bInvert=False))
-+AxisConfig=(AxisKeyName="SteamVR_Knuckles_Right_MiddleGrip",AxisProperties=(DeadZone=0.000000,Sensitivity=1.000000,Exponent=1.000000,bInvert=False))
-+AxisConfig=(AxisKeyName="SteamVR_Knuckles_Right_RingGrip",AxisProperties=(DeadZone=0.000000,Sensitivity=1.000000,Exponent=1.000000,bInvert=False))
-+AxisConfig=(AxisKeyName="SteamVR_Knuckles_Right_PinkyGrip",AxisProperties=(DeadZone=0.000000,Sensitivity=1.000000,Exponent=1.000000,bInvert=False))
-+AxisConfig=(AxisKeyName="Vive_Left_Trigger_Axis",AxisProperties=(DeadZone=0.000000,Sensitivity=1.000000,Exponent=1.000000,bInvert=False))
-+AxisConfig=(AxisKeyName="Vive_Left_Trackpad_X",AxisProperties=(DeadZone=0.000000,Sensitivity=1.000000,Exponent=1.000000,bInvert=False))
-+AxisConfig=(AxisKeyName="Vive_Left_Trackpad_Y",AxisProperties=(DeadZone=0.000000,Sensitivity=1.000000,Exponent=1.000000,bInvert=False))
-+AxisConfig=(AxisKeyName="Vive_Right_Trigger_Axis",AxisProperties=(DeadZone=0.000000,Sensitivity=1.000000,Exponent=1.000000,bInvert=False))
-+AxisConfig=(AxisKeyName="Vive_Right_Trackpad_X",AxisProperties=(DeadZone=0.000000,Sensitivity=1.000000,Exponent=1.000000,bInvert=False))
-+AxisConfig=(AxisKeyName="Vive_Right_Trackpad_Y",AxisProperties=(DeadZone=0.000000,Sensitivity=1.000000,Exponent=1.000000,bInvert=False))
-+AxisConfig=(AxisKeyName="MixedReality_Left_Trigger_Axis",AxisProperties=(DeadZone=0.000000,Sensitivity=1.000000,Exponent=1.000000,bInvert=False))
-+AxisConfig=(AxisKeyName="MixedReality_Left_Thumbstick_X",AxisProperties=(DeadZone=0.000000,Sensitivity=1.000000,Exponent=1.000000,bInvert=False))
-+AxisConfig=(AxisKeyName="MixedReality_Left_Thumbstick_Y",AxisProperties=(DeadZone=0.000000,Sensitivity=1.000000,Exponent=1.000000,bInvert=False))
-+AxisConfig=(AxisKeyName="MixedReality_Left_Trackpad_X",AxisProperties=(DeadZone=0.000000,Sensitivity=1.000000,Exponent=1.000000,bInvert=False))
-+AxisConfig=(AxisKeyName="MixedReality_Left_Trackpad_Y",AxisProperties=(DeadZone=0.000000,Sensitivity=1.000000,Exponent=1.000000,bInvert=False))
-+AxisConfig=(AxisKeyName="MixedReality_Right_Trigger_Axis",AxisProperties=(DeadZone=0.000000,Sensitivity=1.000000,Exponent=1.000000,bInvert=False))
-+AxisConfig=(AxisKeyName="MixedReality_Right_Thumbstick_X",AxisProperties=(DeadZone=0.000000,Sensitivity=1.000000,Exponent=1.000000,bInvert=False))
-+AxisConfig=(AxisKeyName="MixedReality_Right_Thumbstick_Y",AxisProperties=(DeadZone=0.000000,Sensitivity=1.000000,Exponent=1.000000,bInvert=False))
-+AxisConfig=(AxisKeyName="MixedReality_Right_Trackpad_X",AxisProperties=(DeadZone=0.000000,Sensitivity=1.000000,Exponent=1.000000,bInvert=False))
-+AxisConfig=(AxisKeyName="MixedReality_Right_Trackpad_Y",AxisProperties=(DeadZone=0.000000,Sensitivity=1.000000,Exponent=1.000000,bInvert=False))
-+AxisConfig=(AxisKeyName="OculusGo_Left_Trackpad_X",AxisProperties=(DeadZone=0.000000,Sensitivity=1.000000,Exponent=1.000000,bInvert=False))
-+AxisConfig=(AxisKeyName="OculusGo_Left_Trackpad_Y",AxisProperties=(DeadZone=0.000000,Sensitivity=1.000000,Exponent=1.000000,bInvert=False))
-+AxisConfig=(AxisKeyName="OculusGo_Right_Trackpad_X",AxisProperties=(DeadZone=0.000000,Sensitivity=1.000000,Exponent=1.000000,bInvert=False))
-+AxisConfig=(AxisKeyName="OculusGo_Right_Trackpad_Y",AxisProperties=(DeadZone=0.000000,Sensitivity=1.000000,Exponent=1.000000,bInvert=False))
-+AxisConfig=(AxisKeyName="OculusTouch_Left_Grip_Axis",AxisProperties=(DeadZone=0.000000,Sensitivity=1.000000,Exponent=1.000000,bInvert=False))
-+AxisConfig=(AxisKeyName="OculusTouch_Left_Trigger_Axis",AxisProperties=(DeadZone=0.000000,Sensitivity=1.000000,Exponent=1.000000,bInvert=False))
-+AxisConfig=(AxisKeyName="OculusTouch_Left_Thumbstick_X",AxisProperties=(DeadZone=0.000000,Sensitivity=1.000000,Exponent=1.000000,bInvert=False))
-+AxisConfig=(AxisKeyName="OculusTouch_Left_Thumbstick_Y",AxisProperties=(DeadZone=0.000000,Sensitivity=1.000000,Exponent=1.000000,bInvert=False))
-+AxisConfig=(AxisKeyName="OculusTouch_Right_Grip_Axis",AxisProperties=(DeadZone=0.000000,Sensitivity=1.000000,Exponent=1.000000,bInvert=False))
-+AxisConfig=(AxisKeyName="OculusTouch_Right_Trigger_Axis",AxisProperties=(DeadZone=0.000000,Sensitivity=1.000000,Exponent=1.000000,bInvert=False))
-+AxisConfig=(AxisKeyName="OculusTouch_Right_Thumbstick_X",AxisProperties=(DeadZone=0.000000,Sensitivity=1.000000,Exponent=1.000000,bInvert=False))
-+AxisConfig=(AxisKeyName="OculusTouch_Right_Thumbstick_Y",AxisProperties=(DeadZone=0.000000,Sensitivity=1.000000,Exponent=1.000000,bInvert=False))
-+AxisConfig=(AxisKeyName="ValveIndex_Left_Grip_Axis",AxisProperties=(DeadZone=0.000000,Sensitivity=1.000000,Exponent=1.000000,bInvert=False))
-+AxisConfig=(AxisKeyName="ValveIndex_Left_Grip_Force",AxisProperties=(DeadZone=0.000000,Sensitivity=1.000000,Exponent=1.000000,bInvert=False))
-+AxisConfig=(AxisKeyName="ValveIndex_Left_Trigger_Axis",AxisProperties=(DeadZone=0.000000,Sensitivity=1.000000,Exponent=1.000000,bInvert=False))
-+AxisConfig=(AxisKeyName="ValveIndex_Left_Thumbstick_X",AxisProperties=(DeadZone=0.000000,Sensitivity=1.000000,Exponent=1.000000,bInvert=False))
-+AxisConfig=(AxisKeyName="ValveIndex_Left_Thumbstick_Y",AxisProperties=(DeadZone=0.000000,Sensitivity=1.000000,Exponent=1.000000,bInvert=False))
-+AxisConfig=(AxisKeyName="ValveIndex_Left_Trackpad_X",AxisProperties=(DeadZone=0.000000,Sensitivity=1.000000,Exponent=1.000000,bInvert=False))
-+AxisConfig=(AxisKeyName="ValveIndex_Left_Trackpad_Y",AxisProperties=(DeadZone=0.000000,Sensitivity=1.000000,Exponent=1.000000,bInvert=False))
-+AxisConfig=(AxisKeyName="ValveIndex_Left_Trackpad_Force",AxisProperties=(DeadZone=0.000000,Sensitivity=1.000000,Exponent=1.000000,bInvert=False))
-+AxisConfig=(AxisKeyName="ValveIndex_Left_Trackpad_Touch",AxisProperties=(DeadZone=0.000000,Sensitivity=1.000000,Exponent=1.000000,bInvert=False))
-+AxisConfig=(AxisKeyName="ValveIndex_Right_Grip_Axis",AxisProperties=(DeadZone=0.000000,Sensitivity=1.000000,Exponent=1.000000,bInvert=False))
-+AxisConfig=(AxisKeyName="ValveIndex_Right_Grip_Force",AxisProperties=(DeadZone=0.000000,Sensitivity=1.000000,Exponent=1.000000,bInvert=False))
-+AxisConfig=(AxisKeyName="ValveIndex_Right_Trigger_Axis",AxisProperties=(DeadZone=0.000000,Sensitivity=1.000000,Exponent=1.000000,bInvert=False))
-+AxisConfig=(AxisKeyName="ValveIndex_Right_Thumbstick_X",AxisProperties=(DeadZone=0.000000,Sensitivity=1.000000,Exponent=1.000000,bInvert=False))
-+AxisConfig=(AxisKeyName="ValveIndex_Right_Thumbstick_Y",AxisProperties=(DeadZone=0.000000,Sensitivity=1.000000,Exponent=1.000000,bInvert=False))
-+AxisConfig=(AxisKeyName="ValveIndex_Right_Trackpad_X",AxisProperties=(DeadZone=0.000000,Sensitivity=1.000000,Exponent=1.000000,bInvert=False))
-+AxisConfig=(AxisKeyName="ValveIndex_Right_Trackpad_Y",AxisProperties=(DeadZone=0.000000,Sensitivity=1.000000,Exponent=1.000000,bInvert=False))
-+AxisConfig=(AxisKeyName="ValveIndex_Right_Trackpad_Force",AxisProperties=(DeadZone=0.000000,Sensitivity=1.000000,Exponent=1.000000,bInvert=False))
 bAltEnterTogglesFullscreen=True
 bF11TogglesFullscreen=True
 bUseMouseForTouch=False
 bEnableMouseSmoothing=True
 bEnableFOVScaling=True
 bCaptureMouseOnLaunch=True
+bEnableLegacyInputScales=True
 bAlwaysShowTouchInterface=False
 bShowConsoleOnFourFingerTap=True
 bEnableGestureRecognizer=False
@@ -108,50 +29,19 @@ FOVScale=0.011110
 DoubleClickTime=0.200000
 +ActionMappings=(ActionName="Jump",bShift=False,bCtrl=False,bAlt=False,bCmd=False,Key=SpaceBar)
 +ActionMappings=(ActionName="Jump",bShift=False,bCtrl=False,bAlt=False,bCmd=False,Key=Gamepad_FaceButton_Bottom)
-+ActionMappings=(ActionName="ResetVR",bShift=False,bCtrl=False,bAlt=False,bCmd=False,Key=R)
-+ActionMappings=(ActionName="Jump",bShift=False,bCtrl=False,bAlt=False,bCmd=False,Key=Vive_Left_Trigger_Click)
-+ActionMappings=(ActionName="Jump",bShift=False,bCtrl=False,bAlt=False,bCmd=False,Key=Vive_Right_Trigger_Click)
-+ActionMappings=(ActionName="Jump",bShift=False,bCtrl=False,bAlt=False,bCmd=False,Key=MixedReality_Left_Trigger_Click)
-+ActionMappings=(ActionName="Jump",bShift=False,bCtrl=False,bAlt=False,bCmd=False,Key=MixedReality_Right_Trigger_Click)
-+ActionMappings=(ActionName="Jump",bShift=False,bCtrl=False,bAlt=False,bCmd=False,Key=OculusGo_Left_Trigger_Click)
-+ActionMappings=(ActionName="Jump",bShift=False,bCtrl=False,bAlt=False,bCmd=False,Key=OculusTouch_Left_Trigger_Click)
-+ActionMappings=(ActionName="Jump",bShift=False,bCtrl=False,bAlt=False,bCmd=False,Key=OculusTouch_Right_Trigger_Click)
-+ActionMappings=(ActionName="Jump",bShift=False,bCtrl=False,bAlt=False,bCmd=False,Key=ValveIndex_Left_Trigger_Click)
-+ActionMappings=(ActionName="Jump",bShift=False,bCtrl=False,bAlt=False,bCmd=False,Key=ValveIndex_Right_Trigger_Click)
-+ActionMappings=(ActionName="ResetVR",bShift=False,bCtrl=False,bAlt=False,bCmd=False,Key=Vive_Left_Grip_Click)
-+ActionMappings=(ActionName="ResetVR",bShift=False,bCtrl=False,bAlt=False,bCmd=False,Key=MixedReality_Left_Thumbstick_Click)
-+ActionMappings=(ActionName="ResetVR",bShift=False,bCtrl=False,bAlt=False,bCmd=False,Key=OculusGo_Left_Trackpad_Click)
-+ActionMappings=(ActionName="ResetVR",bShift=False,bCtrl=False,bAlt=False,bCmd=False,Key=OculusTouch_Left_Thumbstick_Click)
-+ActionMappings=(ActionName="ResetVR",bShift=False,bCtrl=False,bAlt=False,bCmd=False,Key=ValveIndex_Left_Thumbstick_Click)
-+AxisMappings=(AxisName="MoveForward",Scale=1.000000,Key=W)
-+AxisMappings=(AxisName="MoveForward",Scale=-1.000000,Key=S)
-+AxisMappings=(AxisName="MoveForward",Scale=1.000000,Key=Up)
-+AxisMappings=(AxisName="MoveForward",Scale=-1.000000,Key=Down)
-+AxisMappings=(AxisName="MoveForward",Scale=1.000000,Key=Gamepad_LeftY)
-+AxisMappings=(AxisName="MoveRight",Scale=-1.000000,Key=A)
-+AxisMappings=(AxisName="MoveRight",Scale=1.000000,Key=D)
-+AxisMappings=(AxisName="MoveRight",Scale=1.000000,Key=Gamepad_LeftX)
-+AxisMappings=(AxisName="TurnRate",Scale=1.000000,Key=Gamepad_RightX)
-+AxisMappings=(AxisName="TurnRate",Scale=-1.000000,Key=Left)
-+AxisMappings=(AxisName="TurnRate",Scale=1.000000,Key=Right)
-+AxisMappings=(AxisName="Turn",Scale=1.000000,Key=MouseX)
-+AxisMappings=(AxisName="LookUpRate",Scale=1.000000,Key=Gamepad_RightY)
-+AxisMappings=(AxisName="LookUp",Scale=-1.000000,Key=MouseY)
-+AxisMappings=(AxisName="TurnRate",Scale=-1.000000,Key=Vive_Right_Trackpad_X)
-+AxisMappings=(AxisName="MoveForward",Scale=1.000000,Key=Vive_Left_Trackpad_Y)
-+AxisMappings=(AxisName="MoveRight",Scale=1.000000,Key=Vive_Left_Trackpad_X)
-+AxisMappings=(AxisName="MoveRight",Scale=1.000000,Key=MixedReality_Left_Thumbstick_X)
-+AxisMappings=(AxisName="MoveRight",Scale=1.000000,Key=OculusGo_Left_Trackpad_X)
-+AxisMappings=(AxisName="MoveForward",Scale=1.000000,Key=MixedReality_Left_Thumbstick_Y)
-+AxisMappings=(AxisName="MoveForward",Scale=1.000000,Key=OculusGo_Left_Trackpad_Y)
-+AxisMappings=(AxisName="TurnRate",Scale=-1.000000,Key=MixedReality_Right_Thumbstick_X)
-+AxisMappings=(AxisName="TurnRate",Scale=-1.000000,Key=OculusTouch_Right_Thumbstick_X)
-+AxisMappings=(AxisName="TurnRate",Scale=-1.000000,Key=ValveIndex_Right_Thumbstick_X)
-+AxisMappings=(AxisName="MoveForward",Scale=1.000000,Key=OculusTouch_Left_Thumbstick_Y)
-+AxisMappings=(AxisName="MoveForward",Scale=1.000000,Key=ValveIndex_Left_Thumbstick_Y)
-+AxisMappings=(AxisName="MoveRight",Scale=1.000000,Key=OculusTouch_Left_Thumbstick_X)
-+AxisMappings=(AxisName="MoveRight",Scale=1.000000,Key=ValveIndex_Left_Thumbstick_X)
++AxisMappings=(AxisName="Move Forward / Backward",Scale=1.000000,Key=W)
++AxisMappings=(AxisName="Move Forward / Backward",Scale=-1.000000,Key=S)
++AxisMappings=(AxisName="Move Forward / Backward",Scale=1.000000,Key=Gamepad_LeftY)
++AxisMappings=(AxisName="Move Right / Left",Scale=-1.000000,Key=A)
++AxisMappings=(AxisName="Move Right / Left",Scale=1.000000,Key=D)
++AxisMappings=(AxisName="Move Right / Left",Scale=1.000000,Key=Gamepad_LeftX)
++AxisMappings=(AxisName="Turn Right / Left Gamepad",Scale=1.000000,Key=Gamepad_RightX)
++AxisMappings=(AxisName="Turn Right / Left Mouse",Scale=1.000000,Key=MouseX)
++AxisMappings=(AxisName="Look Up / Down Gamepad",Scale=1.000000,Key=Gamepad_RightY)
++AxisMappings=(AxisName="Look Up / Down Mouse",Scale=-1.000000,Key=MouseY)
+DefaultPlayerInputClass=/Script/Engine.PlayerInput
+DefaultInputComponentClass=/Script/Engine.InputComponent
 DefaultTouchInterface=/Engine/MobileResources/HUD/DefaultVirtualJoysticks.DefaultVirtualJoysticks
+-ConsoleKeys=Tilde
 +ConsoleKeys=Tilde
 
-
diff --git a/Templates/TP_ThirdPersonBP/Config/TemplateDefs.ini b/Templates/TP_ThirdPersonBP/Config/TemplateDefs.ini
--- a/Templates/TP_ThirdPersonBP/Config/TemplateDefs.ini
+++ b/Templates/TP_ThirdPersonBP/Config/TemplateDefs.ini
@@ -41,8 +41,6 @@ FilenameReplacements=(Extensions=("cpp","h","ini","cs"),From="%TEMPLATENAME%",To
 ReplacementsInFiles=(Extensions=("cpp","h","ini","cs"),From="%TEMPLATENAME_UPPERCASE%",To="%PROJECTNAME_UPPERCASE%",bCaseSensitive=true)
 ReplacementsInFiles=(Extensions=("cpp","h","ini","cs"),From="%TEMPLATENAME_LOWERCASE%",To="%PROJECTNAME_LOWERCASE%",bCaseSensitive=true)
 ReplacementsInFiles=(Extensions=("cpp","h","ini","cs"),From="%TEMPLATENAME%",To="%PROJECTNAME%",bCaseSensitive=false)
-SharedContentPacks=(MountName="ThirdPerson",DetailLevels=("Standard"))
-SharedContentPacks=(MountName="Geometry",DetailLevels=("Standard","High"))
-SharedContentPacks=(MountName="Mannequin",DetailLevels=("Standard","High"))
 SharedContentPacks=(MountName="LevelPrototyping",DetailLevels=("High"))
+SharedContentPacks=(MountName="Mannequin",DetailLevels=("High"))
 EditDetailLevelPreference="High"
diff --git a/Templates/TP_ThirdPersonBP/Config/config.ini b/Templates/TP_ThirdPersonBP/Config/config.ini
--- a/Templates/TP_ThirdPersonBP/Config/config.ini
+++ b/Templates/TP_ThirdPersonBP/Config/config.ini
@@ -1,24 +1,16 @@
 [InputSettings]
-+ActionMappings=(ActionName="Jump", Key=SpaceBar)
-+ActionMappings=(ActionName="Jump", Key=Gamepad_FaceButton_Bottom)
-
-+AxisMappings=(AxisName="MoveForward", Key=W, Scale=1.f)
-+AxisMappings=(AxisName="MoveForward", Key=S, Scale=-1.f)
-+AxisMappings=(AxisName="MoveForward", Key=Up, Scale=1.f)
-+AxisMappings=(AxisName="MoveForward", Key=Down, Scale=-1.f)
-+AxisMappings=(AxisName="MoveForward", Key=Gamepad_LeftY, Scale=1.f)
-
-+AxisMappings=(AxisName="MoveRight", Key=A, Scale=-1.f)
-+AxisMappings=(AxisName="MoveRight", Key=D, Scale=1.f)
-+AxisMappings=(AxisName="MoveRight", Key=Gamepad_LeftX, Scale=1.f)
-
-+AxisMappings=(AxisName="TurnRate", Key=Gamepad_RightX, Scale=1.f)
-+AxisMappings=(AxisName="TurnRate", Key=Left, Scale=-1.f)
-+AxisMappings=(AxisName="TurnRate", Key=Right, Scale=1.f)
-+AxisMappings=(AxisName="Turn", Key=MouseX, Scale=1.f)
-
-+AxisMappings=(AxisName="LookUpRate", Key=Gamepad_RightY, Scale=1.f)
-+AxisMappings=(AxisName="LookUp", Key=MouseY, Scale=-1.f)
++ActionMappings=(ActionName="Jump",bShift=False,bCtrl=False,bAlt=False,bCmd=False,Key=SpaceBar)
++ActionMappings=(ActionName="Jump",bShift=False,bCtrl=False,bAlt=False,bCmd=False,Key=Gamepad_FaceButton_Bottom)
++AxisMappings=(AxisName="Move Forward / Backward",Scale=1.000000,Key=W)
++AxisMappings=(AxisName="Move Forward / Backward",Scale=-1.000000,Key=S)
++AxisMappings=(AxisName="Move Forward / Backward",Scale=1.000000,Key=Gamepad_LeftY)
++AxisMappings=(AxisName="Move Right / Left",Scale=-1.000000,Key=A)
++AxisMappings=(AxisName="Move Right / Left",Scale=1.000000,Key=D)
++AxisMappings=(AxisName="Move Right / Left",Scale=1.000000,Key=Gamepad_LeftX)
++AxisMappings=(AxisName="Turn Right / Left Gamepad",Scale=1.000000,Key=Gamepad_RightX)
++AxisMappings=(AxisName="Turn Right / Left Mouse",Scale=1.000000,Key=MouseX)
++AxisMappings=(AxisName="Look Up / Down Gamepad",Scale=1.000000,Key=Gamepad_RightY)
++AxisMappings=(AxisName="Look Up / Down Mouse",Scale=-1.000000,Key=MouseY)
 
 [Redirects]
 GameName=TP_ThirdPersonBP
diff --git a/Templates/TP_ThirdPersonBP/TP_ThirdPersonBP.uproject b/Templates/TP_ThirdPersonBP/TP_ThirdPersonBP.uproject
--- a/Templates/TP_ThirdPersonBP/TP_ThirdPersonBP.uproject
+++ b/Templates/TP_ThirdPersonBP/TP_ThirdPersonBP.uproject
@@ -10,6 +10,52 @@
 			"TargetAllowList": [
 				"Editor"
 			]
+		},
+		{
+			"Name": "OculusVR",
+			"Enabled": false,
+			"SupportedTargetPlatforms": [
+				"Win64",
+				"Android"
+			]
+		},
+		{
+			"Name": "OpenXREyeTracker",
+			"Enabled": false,
+			"SupportedTargetPlatforms": [
+				"Win64",
+				"Linux",
+				"HoloLens",
+				"Android"
+			]
+		},
+		{
+			"Name": "OpenXRHandTracking",
+			"Enabled": false,
+			"SupportedTargetPlatforms": [
+				"Win64",
+				"Linux",
+				"HoloLens",
+				"Android"
+			]
+		},
+		{
+			"Name": "OpenXR",
+			"Enabled": false,
+			"SupportedTargetPlatforms": [
+				"Win64",
+				"Linux",
+				"Android",
+				"HoloLens"
+			]
+		},
+		{
+			"Name": "SteamVR",
+			"Enabled": false,
+			"SupportedTargetPlatforms": [
+				"Win64",
+				"Linux"
+			]
 		}
-	]	
+	]
 }
\ No newline at end of file
diff --git a/Templates/TP_ThirdPersonBP/manifest.json b/Templates/TP_ThirdPersonBP/manifest.json
--- a/Templates/TP_ThirdPersonBP/manifest.json
+++ b/Templates/TP_ThirdPersonBP/manifest.json
@@ -54,7 +54,7 @@
 	],
 	"ClassTypes": "GameMode, Character, SpringArmComponent, CameraComponent, InputComponent",
 	"Category": "BlueprintFeature",
-	"FocusAsset": "/Game/ThirdPersonBP/ThirdPersonOverview.ThirdPersonOverview",
+	"FocusAsset": "/Game/ThirdPerson/Blueprints",
 	"Thumbnail": "TP_ThirdPersonBP.png",
 	"Screenshots":
 	[
@@ -63,17 +63,9 @@
 	"AdditionalFeaturePacks":
 	[
 		{
-			"MountName": "ThirdPerson",
+			"MountName": "LevelPrototyping",
 			"DetailLevels":
 			[
-				"Standard"
-			]
-		},
-		{
-			"MountName": "Geometry",
-			"DetailLevels":
-			[
-				"Standard",
 				"High"
 			]
 		},
@@ -81,7 +73,6 @@
 			"MountName": "Mannequin",
 			"DetailLevels":
 			[
-				"Standard",
 				"High"
 			]
 		}
