Cleaned blueprints and added comments. Updated template descriptions.

diff --git a/Templates/TP_ThirdPerson/Config/TemplateDefs.ini b/Templates/TP_ThirdPerson/Config/TemplateDefs.ini
--- a/Templates/TP_ThirdPerson/Config/TemplateDefs.ini
+++ b/Templates/TP_ThirdPerson/Config/TemplateDefs.ini
@@ -4,7 +4,7 @@ Categories=Games
 
 
 LocalizedDisplayNames=(Language="en",Text="Third Person")
-LocalizedDescriptions=(Language="en",Text="The third person template features a playable character where the camera is positioned behind and slightly above the character. As the character moves using mouse, keyboard, controller or virtual joystick on a touch device, the camera follows the character, observing from an over-the-shoulder point of view. This perspective emphasizes the main character and is popular in action and adventure games. The level contains several objects that the player can jump upon.")
+LocalizedDescriptions=(Language="en",Text="The third person pack features a playable character where the camera is positioned behind and slightly above the character. As the character moves using mouse, keyboard, controller or virtual joystick on a touch device, the camera follows the character, observing from an over-the-shoulder point of view. This perspective emphasizes the main character and is popular in action and adventure games. The template map contains several objects that the player can jump upon, and push around.")
 
 LocalizedDisplayNames=(Language="ko",Text="삼인칭")
 LocalizedDescriptions=(Language="ko",Text="삼인칭 템플릿은 카메라가 머리위 약간 뒤쪽에 위치한 플레이가능 캐릭터를 특징으로 합니다. 마우스, 키보드, 터치 디바이스의 가상 조이스틱으로 캐릭터가 이동함에 따라 카메라는 캐릭터를 따라다니며 어깨위 시점으로 봅니다. 주인공 캐릭터를 강조하는 시점으로, 액션이나 어드벤처 게임에 많이 쓰입니다. 레벨에는 플레이어가 점프할 수 있는 오브젝트가 여럿 있습니다.")
@@ -15,8 +15,8 @@ LocalizedDescriptions=(Language="ja",Text="サードパーソンテンプレー
 LocalizedDisplayNames=(Language="zh-Hans",Text="第三人称游戏")
 LocalizedDescriptions=(Language="zh-Hans",Text="第三人称游戏模板拥有一个可游玩的角色，摄像机位于角色身后偏上的位置。您可以通过鼠标、键盘、手柄或触摸式设备上控制角色移动，摄像机会跟随角色，通过过肩式视角观察世界。此视角突出主角，在动作及冒险类游戏中非常流行。关卡中还包含一些对象，可以让角色跳到上面。")
 
-ClassTypes=GameMode, Character, SpringArmComponent, CameraComponent, InputComponent
-AssetTypes=Animation Sequence, Animation Blueprint, Skeleton, Blend Space 1D, Skeletal Mesh
+ClassTypes=Character, SpringArmComponent, GameMode
+AssetTypes=Animation Sequence, Animation Blueprint, Skeleton, Skeletal Mesh
 
 FoldersToIgnore=Binaries
 FoldersToIgnore=Build
diff --git a/Templates/TP_ThirdPerson/manifest.json b/Templates/TP_ThirdPerson/manifest.json
--- a/Templates/TP_ThirdPerson/manifest.json
+++ b/Templates/TP_ThirdPerson/manifest.json
@@ -23,7 +23,7 @@
 	[
 		{
 			"Language": "en",
-			"Text": "The third person pack features a playable character where the camera is positioned behind and slightly above the character. As the character moves using mouse, keyboard, controller or virtual joystick on a touch device, the camera follows the character, observing from an over-the-shoulder point of view. This perspective emphasizes the main character and is popular in action and adventure games. The example map in the pack contains several objects that the player can jump upon.\n"
+			"Text": "The third person pack features a playable character where the camera is positioned behind and slightly above the character. As the character moves using mouse, keyboard, controller or virtual joystick on a touch device, the camera follows the character, observing from an over-the-shoulder point of view. This perspective emphasizes the main character and is popular in action and adventure games. The template map contains several objects that the player can jump upon, and push around."
 		},
 		{
 			"Language": "ko",
@@ -42,7 +42,7 @@
 	[
 		{
 			"Language": "en",
-			"Text": "Animation Sequence, Animation Blueprint, Skeleton, Blend Space 1D, Skeletal Mesh"
+			"Text": "Animation Sequence, Animation Blueprint, Skeleton, Skeletal Mesh"
 		}
 	],	
 	"SearchTags":
@@ -52,7 +52,7 @@
 			"Text": "Action, Adventure, Third"
 		}
 	],
-	"ClassTypes": "GameMode, Character, SpringArmComponent, CameraComponent, InputComponent",
+	"ClassTypes": "Character, SpringArmComponent, GameMode",
 	"Category": "CodeFeature",
 	"FocusAsset": "/Game/ThirdPerson/Blueprints",
 	"Thumbnail": "TP_ThirdPerson.png",
diff --git a/Templates/TP_ThirdPersonBP/Config/TemplateDefs.ini b/Templates/TP_ThirdPersonBP/Config/TemplateDefs.ini
--- a/Templates/TP_ThirdPersonBP/Config/TemplateDefs.ini
+++ b/Templates/TP_ThirdPersonBP/Config/TemplateDefs.ini
@@ -3,7 +3,7 @@
 Categories=Games
 
 LocalizedDisplayNames=(Language="en",Text="Third Person")
-LocalizedDescriptions=(Language="en",Text="The third person template features a playable character where the camera is positioned behind and slightly above the character. As the character moves using mouse, keyboard, controller or virtual joystick on a touch device, the camera follows the character, observing from an over-the-shoulder point of view. This perspective emphasizes the main character and is popular in action and adventure games. The level contains several objects that the player can jump upon.")
+LocalizedDescriptions=(Language="en",Text="The third person pack features a playable character where the camera is positioned behind and slightly above the character. As the character moves using mouse, keyboard, controller or virtual joystick on a touch device, the camera follows the character, observing from an over-the-shoulder point of view. This perspective emphasizes the main character and is popular in action and adventure games. The template map contains several objects that the player can jump upon, and push around.")
 
 LocalizedDisplayNames=(Language="ko",Text="삼인칭")
 LocalizedDescriptions=(Language="ko",Text="삼인칭 템플릿은 카메라가 머리위 약간 뒤쪽에 위치한 플레이가능 캐릭터를 특징으로 합니다. 마우스, 키보드, 터치 디바이스의 가상 조이스틱으로 캐릭터가 이동함에 따라 카메라는 캐릭터를 따라다니며 어깨위 시점으로 봅니다. 주인공 캐릭터를 강조하는 시점으로, 액션이나 어드벤처 게임에 많이 쓰입니다. 레벨에는 플레이어가 점프할 수 있는 오브젝트가 여럿 있습니다.")
@@ -14,8 +14,8 @@ LocalizedDescriptions=(Language="ja",Text="サードパーソンテンプレー
 LocalizedDisplayNames=(Language="zh-Hans",Text="第三人称游戏")
 LocalizedDescriptions=(Language="zh-Hans",Text="第三人称游戏模板拥有一个可游玩的角色，摄像机位于角色身后偏上的位置。您可以通过鼠标、键盘、手柄或触摸式设备上控制角色移动，摄像机会跟随角色，通过过肩式视角观察世界。此视角突出主角，在动作及冒险类游戏中非常流行。关卡中还包含一些对象，可以让角色跳到上面。")
 
-ClassTypes=GameMode, Character, SpringArmComponent, CameraComponent, InputComponent
-AssetTypes=Animation Sequence, Animation Blueprint, Skeleton, Blend Space 1D, Skeletal Mesh
+ClassTypes=Character, SpringArmComponent, GameMode
+AssetTypes=Animation Sequence, Animation Blueprint, Skeleton, Skeletal Mesh
 
 FoldersToIgnore=Binaries
 FoldersToIgnore=Build
diff --git a/Templates/TP_ThirdPersonBP/manifest.json b/Templates/TP_ThirdPersonBP/manifest.json
--- a/Templates/TP_ThirdPersonBP/manifest.json
+++ b/Templates/TP_ThirdPersonBP/manifest.json
@@ -23,7 +23,7 @@
 	[
 		{
 			"Language": "en",
-			"Text": "The third person pack features a playable character where the camera is positioned behind and slightly above the character. As the character moves using mouse, keyboard, controller or virtual joystick on a touch device, the camera follows the character, observing from an over-the-shoulder point of view. This perspective emphasizes the main character and is popular in action and adventure games. The example map in the pack contains several objects that the player can jump upon.\n"
+			"Text": "The third person pack features a playable character where the camera is positioned behind and slightly above the character. As the character moves using mouse, keyboard, controller or virtual joystick on a touch device, the camera follows the character, observing from an over-the-shoulder point of view. This perspective emphasizes the main character and is popular in action and adventure games. The template map contains several objects that the player can jump upon, and push around."
 		},
 		{
 			"Language": "ko",
@@ -42,7 +42,7 @@
 	[
 		{
 			"Language": "en",
-			"Text": "Animation Sequence, Animation Blueprint, Skeleton, Blend Space 1D, Skeletal Mesh"
+			"Text": "Animation Sequence, Animation Blueprint, Skeleton, Skeletal Mesh"
 		}
 	],	
 	"SearchTags":
@@ -52,7 +52,7 @@
 			"Text": "Action, Adventure, Third"
 		}
 	],
-	"ClassTypes": "GameMode, Character, SpringArmComponent, CameraComponent, InputComponent",
+	"ClassTypes": "Character, SpringArmComponent, GameMode",
 	"Category": "BlueprintFeature",
 	"FocusAsset": "/Game/ThirdPerson/Blueprints",
 	"Thumbnail": "TP_ThirdPersonBP.png",
