[Templates] Addes asset type control rig, IK rig and retargeter to the desciprtion of the third person template and feature pack inis

diff --git a/Templates/TP_ThirdPerson/Config/TemplateDefs.ini b/Templates/TP_ThirdPerson/Config/TemplateDefs.ini
--- a/Templates/TP_ThirdPerson/Config/TemplateDefs.ini
+++ b/Templates/TP_ThirdPerson/Config/TemplateDefs.ini
@@ -19,7 +19,7 @@ LocalizedDescriptions=(Language="zh-Hans",Text="第三人称游戏内容包包
 bThumbnailAsIcon=true
 
 ClassTypes=Character, SpringArmComponent, GameMode
-AssetTypes=Animation Sequence, Animation Blueprint, Skeleton, Skeletal Mesh
+AssetTypes=Animation Sequence, Animation Blueprint, Skeleton, Skeletal Mesh, Control Rig, IK Rig, Retargeter
 
 FoldersToIgnore=Binaries
 FoldersToIgnore=Build
diff --git a/Templates/TP_ThirdPerson/manifest.json b/Templates/TP_ThirdPerson/manifest.json
--- a/Templates/TP_ThirdPerson/manifest.json
+++ b/Templates/TP_ThirdPerson/manifest.json
@@ -42,7 +42,7 @@
 	[
 		{
 			"Language": "en",
-			"Text": "Animation Sequence, Animation Blueprint, Skeleton, Skeletal Mesh"
+			"Text": "Animation Sequence, Animation Blueprint, Skeleton, Skeletal Mesh, Control Rig, IK Rig, Retargeter"
 		}
 	],	
 	"SearchTags":
diff --git a/Templates/TP_ThirdPersonBP/Config/TemplateDefs.ini b/Templates/TP_ThirdPersonBP/Config/TemplateDefs.ini
--- a/Templates/TP_ThirdPersonBP/Config/TemplateDefs.ini
+++ b/Templates/TP_ThirdPersonBP/Config/TemplateDefs.ini
@@ -18,7 +18,7 @@ LocalizedDescriptions=(Language="zh-Hans",Text="第三人称游戏内容包包
 bThumbnailAsIcon=true
 
 ClassTypes=Character, SpringArmComponent, GameMode
-AssetTypes=Animation Sequence, Animation Blueprint, Skeleton, Skeletal Mesh
+AssetTypes=Animation Sequence, Animation Blueprint, Skeleton, Skeletal Mesh, Control Rig, IK Rig, Retargeter
 
 FoldersToIgnore=Binaries
 FoldersToIgnore=Build
diff --git a/Templates/TP_ThirdPersonBP/manifest.json b/Templates/TP_ThirdPersonBP/manifest.json
--- a/Templates/TP_ThirdPersonBP/manifest.json
+++ b/Templates/TP_ThirdPersonBP/manifest.json
@@ -42,7 +42,7 @@
 	[
 		{
 			"Language": "en",
-			"Text": "Animation Sequence, Animation Blueprint, Skeleton, Skeletal Mesh"
+			"Text": "Animation Sequence, Animation Blueprint, Skeleton, Skeletal Mesh, Control Rig, IK Rig, Retargeter"
 		}
 	],	
 	"SearchTags":
