[Template] Switching 3rd Person BP template to enhanced input.

#contentonly

diff --git a/Templates/TP_ThirdPersonBP/Config/DefaultInput.ini b/Templates/TP_ThirdPersonBP/Config/DefaultInput.ini
--- a/Templates/TP_ThirdPersonBP/Config/DefaultInput.ini
+++ b/Templates/TP_ThirdPersonBP/Config/DefaultInput.ini
@@ -68,6 +68,9 @@ bEnableMouseSmoothing=True
 bEnableFOVScaling=True
 bCaptureMouseOnLaunch=True
 bEnableLegacyInputScales=True
+bEnableMotionControls=True
+bFilterInputByPlatformUser=False
+bShouldFlushPressedKeysOnViewportFocusLost=True
 bAlwaysShowTouchInterface=False
 bShowConsoleOnFourFingerTap=True
 bEnableGestureRecognizer=False
@@ -76,25 +79,10 @@ DefaultViewportMouseCaptureMode=CapturePermanently_IncludingInitialMouseDown
 DefaultViewportMouseLockMode=LockOnCapture
 FOVScale=0.011110
 DoubleClickTime=0.200000
-+ActionMappings=(ActionName="Jump",bShift=False,bCtrl=False,bAlt=False,bCmd=False,Key=SpaceBar)
-+ActionMappings=(ActionName="Jump",bShift=False,bCtrl=False,bAlt=False,bCmd=False,Key=Gamepad_FaceButton_Bottom)
-+AxisMappings=(AxisName="Move Forward / Backward",Scale=1.000000,Key=W)
-+AxisMappings=(AxisName="Move Forward / Backward",Scale=-1.000000,Key=S)
-+AxisMappings=(AxisName="Move Forward / Backward",Scale=1.000000,Key=Gamepad_LeftY)
-+AxisMappings=(AxisName="Move Right / Left",Scale=-1.000000,Key=A)
-+AxisMappings=(AxisName="Move Right / Left",Scale=1.000000,Key=D)
-+AxisMappings=(AxisName="Move Right / Left",Scale=1.000000,Key=Gamepad_LeftX)
-+AxisMappings=(AxisName="Turn Right / Left Gamepad",Scale=1.000000,Key=Gamepad_RightX)
-+AxisMappings=(AxisName="Turn Right / Left Mouse",Scale=1.000000,Key=MouseX)
-+AxisMappings=(AxisName="Look Up / Down Gamepad",Scale=1.000000,Key=Gamepad_RightY)
-+AxisMappings=(AxisName="Look Up / Down Mouse",Scale=-1.000000,Key=MouseY)
-+AxisMappings=(AxisName="Move Forward / Backward",Scale=1.000000,Key=Up)
-+AxisMappings=(AxisName="Move Forward / Backward",Scale=-1.000000,Key=Down)
-+AxisMappings=(AxisName="Move Right / Left",Scale=1.000000,Key=Right)
-+AxisMappings=(AxisName="Move Right / Left",Scale=-1.000000,Key=Left)
-DefaultPlayerInputClass=/Script/Engine.PlayerInput
-DefaultInputComponentClass=/Script/Engine.InputComponent
+DefaultPlayerInputClass=/Script/EnhancedInput.EnhancedPlayerInput
+DefaultInputComponentClass=/Script/EnhancedInput.EnhancedInputComponent
 DefaultTouchInterface=/Engine/MobileResources/HUD/DefaultVirtualJoysticks.DefaultVirtualJoysticks
 -ConsoleKeys=Tilde
 +ConsoleKeys=Tilde
++ConsoleKeys=Caret
 
