New Templates - ThirdPerson
Added Feature Pack files
Modified template manifests to include variant content packs
Updated default template map and assets.
#rnx

diff --git a/Templates/TP_ThirdPersonBP/Config/TemplateDefs.ini b/Templates/TP_ThirdPersonBP/Config/TemplateDefs.ini
--- a/Templates/TP_ThirdPersonBP/Config/TemplateDefs.ini
+++ b/Templates/TP_ThirdPersonBP/Config/TemplateDefs.ini
@@ -47,3 +47,16 @@ ReplacementsInFiles=(Extensions=("cpp","h","ini","cs"),From="%TEMPLATENAME%",To=
 SharedContentPacks=(MountName="LevelPrototyping",DetailLevels=("High"))
 SharedContentPacks=(MountName="Characters",DetailLevels=("High"))
 EditDetailLevelPreference="High"
+
+; Variant packs
+; NOTE: make sure these are commented out for release!
+;SharedContentPacks=(MountName="Variant_Combat",DetailLevels=("Standard"))
+;SharedContentPacks=(MountName="Variant_Platforming",DetailLevels=("Standard"))
+;SharedContentPacks=(MountName="Variant_SideScroller",DetailLevels=("Standard"))
+
+; Variant definitions
+Variants=(Name="Combat",LocalizedDisplayNames=((Language="en",Text="Combat")),LocalizedDescriptions=((Language="en",Text="A combat game with melee attacks and AI-controlled enemies.")),SharedContentPacks=((DetailLevels=(Standard),MountName="Variant_Combat")))
+
+Variants=(Name="Platforming",LocalizedDisplayNames=((Language="en",Text="Platforming")),LocalizedDescriptions=((Language="en",Text="A 3D platforming game with tweaked movement properties and advanced jump capabilities.")),SharedContentPacks=((DetailLevels=(Standard),MountName="Variant_Platforming")))
+
+Variants=(Name="SideScroller",LocalizedDisplayNames=((Language="en",Text="Side Scroller")),LocalizedDescriptions=((Language="en",Text="A 2.5D platforming game with plane-constrained movement, advanced jump capabilities and a side-scrolling camera.")),SharedContentPacks=((DetailLevels=(Standard),MountName="Variant_SideScroller")))
\ No newline at end of file
diff --git a/Templates/TP_ThirdPersonBP/TP_ThirdPersonBP.uproject b/Templates/TP_ThirdPersonBP/TP_ThirdPersonBP.uproject
--- a/Templates/TP_ThirdPersonBP/TP_ThirdPersonBP.uproject
+++ b/Templates/TP_ThirdPersonBP/TP_ThirdPersonBP.uproject
@@ -10,6 +10,10 @@
 			"TargetAllowList": [
 				"Editor"
 			]
+		},
+		{
+			"Name": "GameplayStateTree",
+			"Enabled": true
 		}
 	]
 }
\ No newline at end of file
diff --git a/Templates/TemplateResources/Standard/Variant_Combat/FeaturePack/manifest.json b/Templates/TemplateResources/Standard/Variant_Combat/FeaturePack/manifest.json
new file mode 100644
--- /dev/null
+++ b/Templates/TemplateResources/Standard/Variant_Combat/FeaturePack/manifest.json
@@ -0,0 +1,48 @@
+{
+	"Version": 1,
+	"Name":
+	[
+		{
+			"Language": "en",
+			"Text": "Combat Variant Shared Resource Standard"
+		}
+	],
+	"Description":
+	[
+		{
+			"Language": "en",
+			"Text": "Shared Resource Pack for ThirdPerson Combat Variant"
+		}
+	],
+	"AssetTypes":
+	[
+		{
+			"Language": "en",
+			"Text": ""
+		}
+	],
+	"SearchTags":
+	[
+		{
+			"Language": "en",
+			"Text": ""
+		}
+	],	
+	"ClassTypes": " ",
+	"Category": "SharedPack",
+	"Thumbnail": "Variant_Combat.png",
+	"Screenshots":
+	[
+		"Variant_Combat_Preview.png"
+	],
+	"Ident": "Variant_CombatStandard",
+	"AdditionalFiles":
+	{
+		"DestinationFilesFolder": "Variant_Combat",
+		"AdditionalFilesList":
+		[
+			"Templates/TemplateResources/Standard/Variant_Combat/Content/*.*"
+		]
+	}
+}
+
diff --git a/Templates/TemplateResources/Standard/Variant_Platforming/FeaturePack/manifest.json b/Templates/TemplateResources/Standard/Variant_Platforming/FeaturePack/manifest.json
new file mode 100644
--- /dev/null
+++ b/Templates/TemplateResources/Standard/Variant_Platforming/FeaturePack/manifest.json
@@ -0,0 +1,48 @@
+{
+	"Version": 1,
+	"Name":
+	[
+		{
+			"Language": "en",
+			"Text": "Platforming Variant Shared Resource Standard"
+		}
+	],
+	"Description":
+	[
+		{
+			"Language": "en",
+			"Text": "Shared Resource Pack for ThirdPerson Platforming Variant"
+		}
+	],
+	"AssetTypes":
+	[
+		{
+			"Language": "en",
+			"Text": ""
+		}
+	],
+	"SearchTags":
+	[
+		{
+			"Language": "en",
+			"Text": ""
+		}
+	],	
+	"ClassTypes": " ",
+	"Category": "SharedPack",
+	"Thumbnail": "Variant_Platforming.png",
+	"Screenshots":
+	[
+		"Variant_Platforming_Preview.png"
+	],
+	"Ident": "Variant_PlatformingStandard",
+	"AdditionalFiles":
+	{
+		"DestinationFilesFolder": "Variant_Platforming",
+		"AdditionalFilesList":
+		[
+			"Templates/TemplateResources/Standard/Variant_Platforming/Content/*.*"
+		]
+	}
+}
+
diff --git a/Templates/TemplateResources/Standard/Variant_SideScroller/FeaturePack/manifest.json b/Templates/TemplateResources/Standard/Variant_SideScroller/FeaturePack/manifest.json
new file mode 100644
--- /dev/null
+++ b/Templates/TemplateResources/Standard/Variant_SideScroller/FeaturePack/manifest.json
@@ -0,0 +1,48 @@
+{
+	"Version": 1,
+	"Name":
+	[
+		{
+			"Language": "en",
+			"Text": "SideScroller Variant Shared Resource Standard"
+		}
+	],
+	"Description":
+	[
+		{
+			"Language": "en",
+			"Text": "Shared Resource Pack for ThirdPerson SideScroller Variant"
+		}
+	],
+	"AssetTypes":
+	[
+		{
+			"Language": "en",
+			"Text": ""
+		}
+	],
+	"SearchTags":
+	[
+		{
+			"Language": "en",
+			"Text": ""
+		}
+	],	
+	"ClassTypes": " ",
+	"Category": "SharedPack",
+	"Thumbnail": "Variant_SideScroller.png",
+	"Screenshots":
+	[
+		"Variant_SideScroller_Preview.png"
+	],
+	"Ident": "Variant_SideScrollerStandard",
+	"AdditionalFiles":
+	{
+		"DestinationFilesFolder": "Variant_SideScroller",
+		"AdditionalFilesList":
+		[
+			"Templates/TemplateResources/Standard/Variant_SideScroller/Content/*.*"
+		]
+	}
+}
+
