Third Person Template - fix default map center podium position - Sidescrolling Variant - plane constrain NPC - adjust end of level blocking volumes - add pickup effect - add respawn logic to controller - standarize comments

diff --git a/Templates/TP_ThirdPersonBP/Config/TemplateDefs.ini b/Templates/TP_ThirdPersonBP/Config/TemplateDefs.ini
--- a/Templates/TP_ThirdPersonBP/Config/TemplateDefs.ini
+++ b/Templates/TP_ThirdPersonBP/Config/TemplateDefs.ini
@@ -52,7 +52,7 @@ EditDetailLevelPreference="High"
 ; NOTE: make sure these are commented out for release!
 ;SharedContentPacks=(MountName="Variant_Combat",DetailLevels=("Standard"))
 ;SharedContentPacks=(MountName="Variant_Platforming",DetailLevels=("Standard"))
-;SharedContentPacks=(MountName="Variant_SideScroller",DetailLevels=("Standard"))
+SharedContentPacks=(MountName="Variant_SideScroller",DetailLevels=("Standard"))
 
 ; Variant definitions
 Variants=(Name="Combat",LocalizedDisplayNames=((Language="en",Text="Combat")),LocalizedDescriptions=((Language="en",Text="A combat game with melee attacks and AI-controlled enemies.")),SharedContentPacks=((DetailLevels=(Standard),MountName="Variant_Combat")))
