Templates - ThirdPerson
Implemented KillZ and character respawn on Combat and Platforming variants
Adjusted small offset on central platform on default map
Adjusted floor material on default map
Moved default map and template color material to base ThirdPerson folder
Removed Scalar modifier from Move gamepad input
Reorganized and renamed assets on Combat and Platforming templates
Implemented punchable dummy on Combat variant
Replaced destruction fade  on damageable boxes for scale to zero animation
Adjusted collision for central arena mesh on Combat variant
#rnx

diff --git a/Templates/TP_ThirdPersonBP/Config/DefaultEngine.ini b/Templates/TP_ThirdPersonBP/Config/DefaultEngine.ini
--- a/Templates/TP_ThirdPersonBP/Config/DefaultEngine.ini
+++ b/Templates/TP_ThirdPersonBP/Config/DefaultEngine.ini
@@ -2,8 +2,8 @@
 GameName=TP_ThirdPersonBP
 
 [/Script/EngineSettings.GameMapsSettings]
-EditorStartupMap=/Game/ThirdPerson/Maps/ThirdPersonMap.ThirdPersonMap
-GameDefaultMap=/Game/ThirdPerson/Maps/ThirdPersonMap.ThirdPersonMap
+EditorStartupMap=/Game/ThirdPerson/ThirdPersonMap.ThirdPersonMap
+GameDefaultMap=/Game/ThirdPerson/ThirdPersonMap.ThirdPersonMap
 TransitionMap=
 bUseSplitscreen=True
 TwoPlayerSplitscreenLayout=Horizontal
@@ -51,4 +51,5 @@ SoundCueCookQualityIndex=-1
 
 [/Script/LinuxTargetPlatform.LinuxTargetSettings]
 -TargetedRHIs=SF_VULKAN_SM5
-+TargetedRHIs=SF_VULKAN_SM6
\ No newline at end of file
++TargetedRHIs=SF_VULKAN_SM6
+
diff --git a/Templates/TP_ThirdPersonBP/Config/TemplateDefs.ini b/Templates/TP_ThirdPersonBP/Config/TemplateDefs.ini
--- a/Templates/TP_ThirdPersonBP/Config/TemplateDefs.ini
+++ b/Templates/TP_ThirdPersonBP/Config/TemplateDefs.ini
@@ -52,7 +52,7 @@ EditDetailLevelPreference="High"
 ; NOTE: make sure these are commented out for release!
 ;SharedContentPacks=(MountName="Variant_Combat",DetailLevels=("Standard"))
 ;SharedContentPacks=(MountName="Variant_Platforming",DetailLevels=("Standard"))
-SharedContentPacks=(MountName="Variant_SideScroller",DetailLevels=("Standard"))
+;SharedContentPacks=(MountName="Variant_SideScroller",DetailLevels=("Standard"))
 
 ; Variant definitions
 Variants=(Name="Combat",LocalizedDisplayNames=((Language="en",Text="Combat")),LocalizedDescriptions=((Language="en",Text="A combat game with melee attacks and AI-controlled enemies.")),SharedContentPacks=((DetailLevels=(Standard),MountName="Variant_Combat")))
