New Templates - Third Person - Default - clear default Touch Interface - add variant specific touch control - SideScrolling - add touch control - limit jump / push to moveable objects - Platforming - add touch control - Combat - add touch control

diff --git a/Templates/TP_ThirdPersonBP/Config/DefaultInput.ini b/Templates/TP_ThirdPersonBP/Config/DefaultInput.ini
--- a/Templates/TP_ThirdPersonBP/Config/DefaultInput.ini
+++ b/Templates/TP_ThirdPersonBP/Config/DefaultInput.ini
@@ -81,7 +81,7 @@ FOVScale=0.011110
 DoubleClickTime=0.200000
 DefaultPlayerInputClass=/Script/EnhancedInput.EnhancedPlayerInput
 DefaultInputComponentClass=/Script/EnhancedInput.EnhancedInputComponent
-DefaultTouchInterface=/Engine/MobileResources/HUD/DefaultVirtualJoysticks.DefaultVirtualJoysticks
+DefaultTouchInterface=None
 -ConsoleKeys=Tilde
 +ConsoleKeys=Tilde
 
