Fixed animation link issue in third person
Updated environments to single level
Fixed collision on level prototyping assets
Cleaned up folder structure in third person
Fixed grid alignment in shooter

diff --git a/Templates/TP_ThirdPersonBP/Config/DefaultEngine.ini b/Templates/TP_ThirdPersonBP/Config/DefaultEngine.ini
--- a/Templates/TP_ThirdPersonBP/Config/DefaultEngine.ini
+++ b/Templates/TP_ThirdPersonBP/Config/DefaultEngine.ini
@@ -2,8 +2,8 @@
 GameName=TP_ThirdPersonBP
 
 [/Script/EngineSettings.GameMapsSettings]
-EditorStartupMap=/Game/ThirdPerson/ThirdPersonMap.ThirdPersonMap
-GameDefaultMap=/Game/ThirdPerson/ThirdPersonMap.ThirdPersonMap
+EditorStartupMap=/Game/ThirdPerson/Lvl_ThirdPerson.Lvl_ThirdPerson
+GameDefaultMap=/Game/ThirdPerson/Lvl_ThirdPerson.Lvl_ThirdPerson
 TransitionMap=
 bUseSplitscreen=True
 TwoPlayerSplitscreenLayout=Horizontal
@@ -53,3 +53,17 @@ SoundCueCookQualityIndex=-1
 -TargetedRHIs=SF_VULKAN_SM5
 +TargetedRHIs=SF_VULKAN_SM6
 
+[/Script/AndroidFileServerEditor.AndroidFileServerRuntimeSettings]
+bEnablePlugin=True
+bAllowNetworkConnection=True
+SecurityToken=9A8C9D134D11008E5B456E90FC2B7A65
+bIncludeInShipping=False
+bAllowExternalStartInShipping=False
+bCompileAFSProject=False
+bUseCompression=False
+bLogFiles=False
+bReportStats=False
+ConnectionType=USBOnly
+bUseManualIPAddress=False
+ManualIPAddress=
+
diff --git a/Templates/TP_ThirdPersonBP/Config/TemplateDefs.ini b/Templates/TP_ThirdPersonBP/Config/TemplateDefs.ini
--- a/Templates/TP_ThirdPersonBP/Config/TemplateDefs.ini
+++ b/Templates/TP_ThirdPersonBP/Config/TemplateDefs.ini
@@ -46,6 +46,7 @@ ReplacementsInFiles=(Extensions=("cpp","h","ini","cs"),From="%TEMPLATENAME_LOWER
 ReplacementsInFiles=(Extensions=("cpp","h","ini","cs"),From="%TEMPLATENAME%",To="%PROJECTNAME%",bCaseSensitive=false)
 SharedContentPacks=(MountName="LevelPrototyping",DetailLevels=("High"))
 SharedContentPacks=(MountName="Characters",DetailLevels=("High"))
+SharedContentPacks=(MountName="Input",DetailLevels=("High"))
 EditDetailLevelPreference="High"
 
 ; Variant packs
