Uploading CJK localization for Updated Templates in 5.6.1 and also fixed the typo in Top Down template's description
#rnx

diff --git a/Templates/TP_FirstPerson/Config/TemplateDefs.ini b/Templates/TP_FirstPerson/Config/TemplateDefs.ini
--- a/Templates/TP_FirstPerson/Config/TemplateDefs.ini
+++ b/Templates/TP_FirstPerson/Config/TemplateDefs.ini
@@ -5,14 +5,14 @@ Categories=Games
 LocalizedDisplayNames=(Language="en",Text="First Person")
 LocalizedDescriptions=(Language="en",Text="This template features a player character which is viewed from first person perspective. The player can look and move around using mouse, controller or virtual joystick on a touch device. You can select a genre-specific starting point from the Variants drop-down. The C++ version includes code and content for all variants.")
 
-LocalizedDisplayNames=(Language="ko",Text="일인칭")
-LocalizedDescriptions=(Language="ko",Text="일인칭 템플릿은 일인칭 시점에서 보이는 (양팔로 표현되는) 플레이어 캐릭터를 특징으로 합니다. 캐릭터는 레벨 내에서 키보드, 컨트롤러나 터치 디바이스의 가상 조이스틱으로 이동할 수 있습니다. 추가적으로 플레이어는 마우스, 컨트롤러나 터치 디바이스의 가상 조이스틱을 사용하여 주변을 둘러볼 수도 있습니다. 또한, 캐릭터는 총을 장착할 수 있으며 마우스, 컨트롤러나 터치 디바이스의 가상 조이스틱을 사용해 간단한 구형 발사체를 발사하여, 아레나 벽에 튕기면서 레벨에 있는 일부 피직스 오브젝트에 영향을 줍니다.")
+LocalizedDisplayNames=(Language="ko",Text="1인칭")
+LocalizedDescriptions=(Language="ko",Text="이 템플릿은 1인칭 시점으로 보는 플레이어 캐릭터를 제공합니다. 플레이어는 마우스, 컨트롤러 또는 터치 디바이스의 가상 조이스틱을 사용하여 주변을 둘러보고 이동할 수 있습니다. 배리언트 드롭다운에서 장르에 맞게 시작할 옵션을 선택할 수 있습니다. C++ 버전에는 모든 배리언트에 대한 코드와 콘텐츠가 포함되어 있습니다.")
 
-LocalizedDisplayNames=(Language="ja",Text="ファーストパーソン")
-LocalizedDescriptions=(Language="ja",Text="ファーストパーソン テンプレートは、一人称視点から見たプレイヤー キャラクター (両腕が表示される) を特徴とします。タッチ デバイスでキーボード、コントローラー、またはバーチャル ジョイスティックを使用して、レベル内でキャラクターを移動させることができます。さらにプレイヤーは、タッチデバイスでマウス、コントローラー、またはバーチャル ジョイスティックを使って見回すことができます。キャラクターは、タッチデバイスで、マウス、コントローラー、またはバーチャル ジョイスティックを使用して、アリーナの壁で跳ね返りながら、レベル内の一部の物理オブジェクトに影響を与えるシンプルな球体の発射体を発射する銃を拾うこともできます。")
+LocalizedDisplayNames=(Language="ja",Text="一人称")
+LocalizedDescriptions=(Language="ja",Text="このテンプレートには、一人称視点で表示されるプレイヤー キャラクターが含まれています。プレイヤーは、マウス、コントローラー、またはタッチ デバイスのバーチャル ジョイスティックを使用して、周囲を見渡したり、移動したりできます。[バリアント] ドロップダウンからジャンル固有の開始点を選択できます。C++ バージョンには、すべてのバリアントのコードとコンテンツが含まれています")
 
-LocalizedDisplayNames=(Language="zh-Hans",Text="第一人称游戏")
-LocalizedDescriptions=(Language="zh-Hans",Text="第一人称模板包含一个第一人称玩家角色（通过一双手臂呈现）。你可以通过键盘、手柄或触摸式设备上的虚拟摇杆控制角色在关卡中四处移动。你还可以使用鼠标、手柄或触摸式设备上的虚拟摇杆转动角色视角，观察四周。此外，你还可以让角色拾取枪支，然后用鼠标、手柄或触摸式设备上的虚拟摇杆来发射球形子弹。子弹会撞飞关卡中的物理对象，还能从竞技场墙壁上弹开。")
+LocalizedDisplayNames=(Language="zh-Hans",Text="第一人称模板")
+LocalizedDescriptions=(Language="zh-Hans",Text="此模板包含一个以第一人称视角呈现的玩家角色。玩家可以通过鼠标、手柄或触控设备上的虚拟摇杆来转动视角和进行移动。你可以从“变体（Variants）”下拉菜单中选择特定游戏类型的起始点。C++版本包含了所有变体的代码和内容。")
 
 ; This project uses a thumbnail as the project icon
 bThumbnailAsIcon=true
diff --git a/Templates/TP_FirstPerson/manifest.json b/Templates/TP_FirstPerson/manifest.json
--- a/Templates/TP_FirstPerson/manifest.json
+++ b/Templates/TP_FirstPerson/manifest.json
@@ -1,4 +1,4 @@
-{
+﻿{
 	"Version": 1,
 	"Name": [
 		{
@@ -7,15 +7,15 @@
 		},
 		{
 			"Language": "ko",
-			"Text": "일인칭"
+			"Text": "1인칭"
 		},
 		{
 			"Language": "zh-Hans",
-			"Text": "第一人称游戏"
+			"Text": "第一人称模板"
 		},
 		{
 			"Language": "ja",
-			"Text": "ファーストパーソン"
+			"Text": "一人称"
 		}
 	],
 	"Description": [
@@ -25,15 +25,15 @@
 		},
 		{
 			"Language": "ko",
-			"Text": "일인칭 템플릿은 일인칭 시점에서 보이는 (양팔로 표현되는) 플레이어 캐릭터를 특징으로 합니다. 캐릭터는 레벨 내에서 키보드, 컨트롤러나 터치 디바이스의 가상 조이스틱으로 이동할 수 있습니다. 추가적으로 플레이어는 마우스, 컨트롤러나 터치 디바이스의 가상 조이스틱을 사용하여 주변을 둘러볼 수도 있습니다. 또한, 캐릭터는 총을 장착할 수 있으며 마우스, 컨트롤러나 터치 디바이스의 가상 조이스틱을 사용해 간단한 구형 발사체를 발사하여, 아레나 벽에 튕기면서 레벨에 있는 일부 피직스 오브젝트에 영향을 줍니다."
+			"Text": "이 팩에는 1인칭 시점으로 보는 플레이어 캐릭터가 등장합니다. 플레이어는 마우스, 컨트롤러 또는 터치 디바이스의 가상 조이스틱을 사용하여 주변을 둘러보고 이동할 수 있습니다."
 		},
 		{
 			"Language": "zh-Hans",
-			"Text": "第一人称模板包含一个第一人称玩家角色（通过一双手臂呈现）。你可以通过键盘、手柄或触摸式设备上的虚拟摇杆控制角色在关卡中四处移动。你还可以使用鼠标、手柄或触摸式设备上的虚拟摇杆转动角色视角，观察四周。此外，你还可以让角色拾取枪支，然后用鼠标、手柄或触摸式设备上的虚拟摇杆来发射球形子弹。子弹会撞飞关卡中的物理对象，还能从竞技场墙壁上弹开。"
+			"Text": "此模板包含一个以第一人称视角呈现的玩家角色。玩家可以通过鼠标、控制器或触控设备上的虚拟摇杆来观察和移动。"
 		},
 		{
 			"Language": "ja",
-			"Text": "ファーストパーソン テンプレートは、一人称視点から見たプレイヤー キャラクター (両腕が表示される) を特徴とします。タッチ デバイスでキーボード、コントローラー、またはバーチャル ジョイスティックを使用して、レベル内でキャラクターを移動させることができます。さらにプレイヤーは、タッチデバイスでマウス、コントローラー、またはバーチャル ジョイスティックを使って見回すことができます。キャラクターは、タッチデバイスで、マウス、コントローラー、またはバーチャル ジョイスティックを使用して、アリーナの壁で跳ね返りながら、レベル内の一部の物理オブジェクトに影響を与えるシンプルな球体の発射体を発射する銃を拾うこともできます。"
+			"Text": "このパックには、一人称視点で表示されるプレイヤー キャラクターが含まれています。プレイヤーは、マウス、コントローラー、またはタッチ デバイスのバーチャル ジョイスティックを使用して、周囲を見渡したり、移動したりできます"
 		}
 	],
 	"AssetTypes": [
diff --git a/Templates/TP_FirstPersonBP/Config/TemplateDefs.ini b/Templates/TP_FirstPersonBP/Config/TemplateDefs.ini
--- a/Templates/TP_FirstPersonBP/Config/TemplateDefs.ini
+++ b/Templates/TP_FirstPersonBP/Config/TemplateDefs.ini
@@ -5,6 +5,15 @@ Categories=Games
 LocalizedDisplayNames=(Language="en",Text="First Person")
 LocalizedDescriptions=(Language="en",Text="This template features a player character which is viewed from first person perspective. The player can look and move around using mouse, controller or virtual joystick on a touch device. You can select a genre-specific starting point from the Variants drop-down. The C++ version includes code and content for all variants.")
 
+LocalizedDisplayNames=(Language="ko",Text="1인칭")
+LocalizedDescriptions=(Language="ko",Text="이 템플릿은 1인칭 시점으로 보는 플레이어 캐릭터를 제공합니다. 플레이어는 마우스, 컨트롤러 또는 터치 디바이스의 가상 조이스틱을 사용하여 주변을 둘러보고 이동할 수 있습니다. 배리언트 드롭다운에서 장르에 맞게 시작할 옵션을 선택할 수 있습니다. C++ 버전에는 모든 배리언트에 대한 코드와 콘텐츠가 포함되어 있습니다.")
+
+LocalizedDisplayNames=(Language="ja",Text="一人称")
+LocalizedDescriptions=(Language="ja",Text="このテンプレートには、一人称視点で表示されるプレイヤー キャラクターが含まれています。プレイヤーは、マウス、コントローラー、またはタッチ デバイスのバーチャル ジョイスティックを使用して、周囲を見渡したり、移動したりできます。[バリアント] ドロップダウンからジャンル固有の開始点を選択できます。C++ バージョンには、すべてのバリアントのコードとコンテンツが含まれています")
+
+LocalizedDisplayNames=(Language="zh-Hans",Text="第一人称模板")
+LocalizedDescriptions=(Language="zh-Hans",Text="此模板包含一个以第一人称视角呈现的玩家角色。玩家可以通过鼠标、手柄或触控设备上的虚拟摇杆来转动视角和进行移动。你可以从“变体（Variants）”下拉菜单中选择特定游戏类型的起始点。C++版本包含了所有变体的代码和内容。")
+
 ; This project uses a thumbnail as the project icon
 bThumbnailAsIcon=true
 
@@ -46,8 +55,8 @@ EditDetailLevelPreference="High"
 ;SharedContentPacks=(MountName="Weapons",DetailLevels=("Standard"))
 
 ; Variant definitions
-Variants=(Name="ArenaShooter",LocalizedDisplayNames=((Language="en",Text="Arena Shooter")),LocalizedDescriptions=((Language="en",Text="An arena shooter with weapon pickups and enemy AI opponents.")),SharedContentPacks=((DetailLevels=(Standard),MountName="Weapons"),(DetailLevels=(Standard),MountName="Variant_Shooter")) )
+Variants=(Name="ArenaShooter",LocalizedDisplayNames=((Language="en",Text="Arena Shooter"),(Language="ja",Text="アリーナ シューター"),(Language="ko",Text="아레나 슈터"),(Language="zh-Hans",Text="竞技场射击游戏模板")),LocalizedDescriptions=((Language="en",Text="An arena shooter with weapon pickups and enemy AI opponents."),(Language="ja",Text="複数の武器のピックアップ、AI 対戦相手を含むアリーナ シューターです"),(Language="ko",Text="무기 및 적 AI 상대가 포함된 아레나 슈터입니다."),(Language="zh-Hans",Text="一个拥有武器拾取功能和敌对人工智能对手的竞技场射击游戏。")),SharedContentPacks=((DetailLevels=(Standard),MountName="Weapons"),(DetailLevels=(Standard),MountName="Variant_Shooter")) )
 
-Variants=(Name="SurvivalHorror",LocalizedDisplayNames=((Language="en",Text="Survival Horror")),LocalizedDescriptions=((Language="en",Text="A simple first person character with sprint mechanic and torch set in a dark and suspenseful map.")),SharedContentPacks=((DetailLevels=(Standard),MountName="Variant_Horror")) )
+Variants=(Name="SurvivalHorror",LocalizedDisplayNames=((Language="en",Text="Survival Horror"),(Language="ja",Text="サバイバルホラー"),(Language="ko",Text="서바이벌 호러"),(Language="zh-Hans",Text="生存恐怖游戏模板")),LocalizedDescriptions=((Language="en",Text="A simple first person character with sprint mechanic and torch set in a dark and suspenseful map."),(Language="ja",Text="スプリント メカニクスとトーチを備えた、暗くサスペンスに満ちたマップを舞台にしたシンプルな一人称キャラクターです"),(Language="ko",Text="어둡고 긴장감 넘치는 맵을 배경으로 스프린트 메카닉과 횃불을 든 단순한 1인칭 캐릭터가 등장합니다."),(Language="zh-Hans",Text="一个简单的第一人称角色，具备冲刺机制和火把功能，置身于之中一个黑暗且充满悬疑氛围的地图。")),SharedContentPacks=((DetailLevels=(Standard),MountName="Variant_Horror")) )
 
-Variants=(Name="All",LocalizedDisplayNames=((Language="en",Text="All")),LocalizedDescriptions=((Language="en",Text="Includes all First Person variants.")),SharedContentPacks=((DetailLevels=(Standard),MountName="Weapons"),(DetailLevels=(Standard),MountName="Variant_Horror"),(DetailLevels=(Standard),MountName="Variant_Shooter")) )
\ No newline at end of file
+Variants=(Name="All",LocalizedDisplayNames=((Language="en",Text="All"),(Language="ja",Text="All (全員)"),(Language="ko",Text="모두"),(Language="zh-Hans",Text="全部")),LocalizedDescriptions=((Language="en",Text="Includes all First Person variants."),(Language="ja",Text="すべての一人称バリアントを含みます"),(Language="ko",Text="모든 1인칭 베리언트를 포함합니다."),(Language="zh-Hans",Text="包含全部第一人称模板变体。")),SharedContentPacks=((DetailLevels=(Standard),MountName="Weapons"),(DetailLevels=(Standard),MountName="Variant_Horror"),(DetailLevels=(Standard),MountName="Variant_Shooter")) )
\ No newline at end of file
diff --git a/Templates/TP_FirstPersonBP/manifest.json b/Templates/TP_FirstPersonBP/manifest.json
--- a/Templates/TP_FirstPersonBP/manifest.json
+++ b/Templates/TP_FirstPersonBP/manifest.json
@@ -1,4 +1,4 @@
-{
+﻿{
 	"Version": 1,
 	"Name": [
 		{
@@ -7,15 +7,15 @@
 		},
 		{
 			"Language": "ko",
-			"Text": "일인칭"
+			"Text": "1인칭"
 		},
 		{
 			"Language": "zh-Hans",
-			"Text": "第一人称游戏"
+			"Text": "第一人称模板"
 		},
 		{
 			"Language": "ja",
-			"Text": "ファーストパーソン"
+			"Text": "一人称"
 		}
 	],
 	"Description": [
@@ -25,15 +25,15 @@
 		},
 		{
 			"Language": "ko",
-			"Text": "일인칭 템플릿은 일인칭 시점에서 보이는 (양팔로 표현되는) 플레이어 캐릭터를 특징으로 합니다. 캐릭터 이동은 키보드, 컨트롤러, 터치 디바이스의 가상 조이스틱으로 가능합니다. 추가적으로 플레이어는 마우스, 컨트롤러, 터치 디바이스의 가상 조이스틱을 사용하여 주변을 둘러볼 수도 있습니다. 캐릭터에는 총이 장착되어 있어 마우스, 컨트롤러, 터치 디바이스의 가상 조이스틱을 사용하면 간단한 발사체가 발사되어, 아레나 벽에 튕기면서 레벨에 있는 몇몇 물리 오브젝트에 영향을 줍니다."
+			"Text": "이 팩에는 1인칭 시점으로 보는 플레이어 캐릭터가 등장합니다. 플레이어는 마우스, 컨트롤러 또는 터치 디바이스의 가상 조이스틱을 사용하여 주변을 둘러보고 이동할 수 있습니다."
 		},
 		{
 			"Language": "zh-Hans",
-			"Text": "第一人称模板包含一个第一人称玩家角色（通过一双手臂呈现）。你可以通过键盘、手柄或触摸式设备上的虚拟摇杆控制角色在关卡中四处移动。你还可以使用鼠标、手柄或触摸式设备上的虚拟摇杆转动角色视角，观察四周。此外，你还可以让角色拾取枪支，然后用鼠标、手柄或触摸式设备上的虚拟摇杆来发射球形子弹。子弹会撞飞关卡中的物理对象，还能从竞技场墙壁上弹开。"
+			"Text": "此模板包含一个以第一人称视角呈现的玩家角色。玩家可以通过鼠标、控制器或触控设备上的虚拟摇杆来观察和移动。"
 		},
 		{
 			"Language": "ja",
-			"Text": "ファーストパーソンテンプレートには、一人称視点で表示されるプレイヤーキャラクタ（腕として表示されています）がいます。キーボード、コントローラー、タッチデバイスでのバーチャルジョイスティックでこのキャラクタを移動させることができます。追加して、プレイヤーはマウス、コントローラー、タッチデバイスでのバーチャルジョイスティックで周りを見回すことができます。さらにキャラクタは銃を装備して、マウス、コントローラー、タッチデバイスでのバーチャルジョイスティックを使ってシンプルな球体の弾を発射できます。この弾は物理オブジェクトに影響し、アリーナの壁で反射します。"
+			"Text": "このパックには、一人称視点で表示されるプレイヤー キャラクターが含まれています。プレイヤーは、マウス、コントローラー、またはタッチ デバイスのバーチャル ジョイスティックを使用して、周囲を見渡したり、移動したりできます"
 		}
 	],
 	"AssetTypes": [
diff --git a/Templates/TP_ThirdPerson/Config/TemplateDefs.ini b/Templates/TP_ThirdPerson/Config/TemplateDefs.ini
--- a/Templates/TP_ThirdPerson/Config/TemplateDefs.ini
+++ b/Templates/TP_ThirdPerson/Config/TemplateDefs.ini
@@ -1,4 +1,4 @@
-[/Script/GameProjectGeneration.TemplateProjectDefs]
+﻿[/Script/GameProjectGeneration.TemplateProjectDefs]
 
 Categories=Games
 
@@ -6,14 +6,14 @@ Categories=Games
 LocalizedDisplayNames=(Language="en",Text="Third Person")
 LocalizedDescriptions=(Language="en",Text="This template features a playable character where the camera observes an over-the-shoulder point of view. The player can rotate the camera and move around using keyboard, mouse, controller or virtual joystick on a touch device. You can select a genre-specific starting point from the Variants drop-down. The C++ version includes code and content for all variants.")
 
-LocalizedDisplayNames=(Language="ko",Text="삼인칭")
-LocalizedDescriptions=(Language="ko",Text="삼인칭 템플릿에는 플레이 가능한 캐릭터가 포함되어 있으며 카메라가 캐릭터 머리 위 약간 뒤쪽에 위치합니다. 캐릭터가 마우스, 키보드, 터치 디바이스의 가상 조이스틱으로 움직일 때 카메라는 어깨너머의 관점에서 캐릭터를 관찰하며 따라갑니다. 이 시점은 주인공 캐릭터를 강조해 액션이나 어드벤처 게임에 주로 활용됩니다. 이 템플릿 맵에는 플레이어가 점프하여 착지하고 밀 수 있는 여러 가지 오브젝트가 포함되어 있습니다.")
+LocalizedDisplayNames=(Language="ko",Text="3인칭")
+LocalizedDescriptions=(Language="ko",Text="이 템플릿은 카메라가 어깨 너머 시점으로 관찰하는 플레이 가능한 캐릭터를 제공합니다. 플레이어는 키보드, 마우스, 컨트롤러 또는 터치 디바이스의 가상 조이스틱을 사용하여 카메라를 회전하고 이동할 수 있습니다. 배리언트 드롭다운에서 장르에 맞게 시작할 옵션을 선택할 수 있습니다. C++ 버전에는 모든 배리언트에 대한 코드와 콘텐츠가 포함되어 있습니다.")
 
-LocalizedDisplayNames=(Language="ja",Text="サードパーソン")
-LocalizedDescriptions=(Language="ja",Text="サードパーソン パックは、後ろ側と少し上にカメラが配置されているプレイ可能なキャラクターを備えています。タッチデバイスでマウス、キーボード、コントローラー、またはバーチャル ジョイスティックの使用によりキャラクターが移動すると、カメラはキャラクターを追跡し、肩越しの視点から観察します。この視点では主人公が強調されるため、アクションゲームやアドベンチャー ゲームで人気があります。テンプレート マップには、プレイヤーがジャンプしたり押すことができるオブジェクトがいくつか含まれています。")
+LocalizedDisplayNames=(Language="ja",Text="三人称")
+LocalizedDescriptions=(Language="ja",Text="このテンプレートには、カメラが肩越しの視点で表示されるプレイアブル キャラクターが含まれています。プレイヤーは、キーボード、マウス、コントローラー、またはタッチ デバイスのバーチャル ジョイスティックを使用して、カメラを回転させたり移動したりできます。[バリアント] ドロップダウンからジャンル固有の開始点を選択できます。C++ バージョンには、すべてのバリアントのコードとコンテンツが含まれています")
 
-LocalizedDisplayNames=(Language="zh-Hans",Text="第三人称游戏")
-LocalizedDescriptions=(Language="zh-Hans",Text="第三人称游戏内容包包含一个可操控的角色以及一台位于角色身后斜上方的摄像机。你可以通过鼠标、键盘、手柄或触摸式设备上的虚拟摇杆控制角色移动。摄像机会跟随角色，以越肩视角拍摄场景。这种视角突出主角，常用于动作及冒险类游戏。该模板的地图中还包含一些对象，可以让角色跳上跳下或四处推动。")
+LocalizedDisplayNames=(Language="zh-Hans",Text="第三人称模板")
+LocalizedDescriptions=(Language="zh-Hans",Text="此模板包含一个可操控的角色，镜头采用过肩视角进行拍摄。玩家可以通过键盘、鼠标、手柄或触控设备上的虚拟摇杆来旋转镜头并进行移动。你可以从“变体（Variants）”下拉菜单中选择特定游戏类型的起始点。C++版本包含了所有变体的代码和内容。")
 
 ; This project uses a thumbnail as the project icon
 bThumbnailAsIcon=true
diff --git a/Templates/TP_ThirdPerson/manifest.json b/Templates/TP_ThirdPerson/manifest.json
--- a/Templates/TP_ThirdPerson/manifest.json
+++ b/Templates/TP_ThirdPerson/manifest.json
@@ -1,4 +1,4 @@
-{
+﻿{
 	"Version": 1,
 	"Name": [
 		{
@@ -7,15 +7,15 @@
 		},
 		{
 			"Language": "ko",
-			"Text": "삼인칭"
+			"Text": "3인칭"
 		},
 		{
 			"Language": "zh-Hans",
-			"Text": "第三人称游戏"
+			"Text": "第三人称模板"
 		},
 		{
 			"Language": "ja",
-			"Text": "サードパーソン"
+			"Text": "三人称"
 		}
 	],
 	"Description": [
@@ -25,15 +25,15 @@
 		},
 		{
 			"Language": "ko",
-			"Text": "삼인칭 템플릿에는 플레이 가능한 캐릭터가 포함되어 있으며 카메라가 캐릭터 머리 위 약간 뒤쪽에 위치합니다. 캐릭터가 마우스, 키보드, 터치 디바이스의 가상 조이스틱으로 움직일 때 카메라는 어깨너머의 관점에서 캐릭터를 관찰하며 따라갑니다. 이 시점은 주인공 캐릭터를 강조해 액션이나 어드벤처 게임에 주로 활용됩니다. 이 템플릿 맵에는 플레이어가 점프하여 착지하고 밀 수 있는 여러 가지 오브젝트가 포함되어 있습니다."
+			"Text": "이 팩에는 카메라가 어깨 너머 시점으로 관찰하는 플레이 가능한 캐릭터가 등장합니다. 플레이어는 키보드, 마우스, 컨트롤러 또는 터치 디바이스의 가상 조이스틱을 사용하여 카메라를 회전하고 이동할 수 있습니다."
 		},
 		{
 			"Language": "zh-Hans",
-			"Text": "第三人称游戏内容包包含一个可操控的角色以及一台位于角色身后斜上方的摄像机。你可以通过鼠标、键盘、手柄或触摸式设备上的虚拟摇杆控制角色移动。摄像机会跟随角色，以越肩视角拍摄场景。这种视角突出主角，常用于动作及冒险类游戏。该模板的地图中还包含一些对象，可以让角色跳上跳下或四处推动。"
+			"Text": "此模板包含一个可操控的角色，镜头采用过肩视角进行拍摄。玩家可以通过键盘、鼠标、手柄或触控设备上的虚拟摇杆来旋转镜头并进行移动。"
 		},
 		{
 			"Language": "ja",
-			"Text": "サードパーソン パックは、後ろ側と少し上にカメラが配置されているプレイ可能なキャラクターを備えています。タッチデバイスでマウス、キーボード、コントローラー、またはバーチャル ジョイスティックの使用によりキャラクターが移動すると、カメラはキャラクターを追跡し、肩越しの視点から観察します。この視点では主人公が強調されるため、アクションゲームやアドベンチャー ゲームで人気があります。テンプレート マップには、プレイヤーがジャンプしたり押すことができるオブジェクトがいくつか含まれています。"
+			"Text": "このパックには、カメラが肩越しの視点で表示されるプレイアブル キャラクターが含まれています。プレイヤーは、キーボード、マウス、コントローラー、またはタッチ デバイスのバーチャル ジョイスティックを使用して、カメラを回転させたり移動したりできます"
 		}
 	],
 	"AssetTypes": [
diff --git a/Templates/TP_ThirdPersonBP/Config/TemplateDefs.ini b/Templates/TP_ThirdPersonBP/Config/TemplateDefs.ini
--- a/Templates/TP_ThirdPersonBP/Config/TemplateDefs.ini
+++ b/Templates/TP_ThirdPersonBP/Config/TemplateDefs.ini
@@ -5,6 +5,15 @@ Categories=Games
 LocalizedDisplayNames=(Language="en",Text="Third Person")
 LocalizedDescriptions=(Language="en",Text="This template features a playable character where the camera observes an over-the-shoulder point of view. The player can rotate the camera and move around using keyboard, mouse, controller or virtual joystick on a touch device. You can select a genre-specific starting point from the Variants drop-down. The C++ version includes code and content for all variants.")
 
+LocalizedDisplayNames=(Language="ko",Text="3인칭")
+LocalizedDescriptions=(Language="ko",Text="이 템플릿은 카메라가 어깨 너머 시점으로 관찰하는 플레이 가능한 캐릭터를 제공합니다. 플레이어는 키보드, 마우스, 컨트롤러 또는 터치 디바이스의 가상 조이스틱을 사용하여 카메라를 회전하고 이동할 수 있습니다. 배리언트 드롭다운에서 장르에 맞게 시작할 옵션을 선택할 수 있습니다. C++ 버전에는 모든 배리언트에 대한 코드와 콘텐츠가 포함되어 있습니다.")
+
+LocalizedDisplayNames=(Language="ja",Text="三人称")
+LocalizedDescriptions=(Language="ja",Text="このテンプレートには、カメラが肩越しの視点で表示されるプレイアブル キャラクターが含まれています。プレイヤーは、キーボード、マウス、コントローラー、またはタッチ デバイスのバーチャル ジョイスティックを使用して、カメラを回転させたり移動したりできます。[バリアント] ドロップダウンからジャンル固有の開始点を選択できます。C++ バージョンには、すべてのバリアントのコードとコンテンツが含まれています")
+
+LocalizedDisplayNames=(Language="zh-Hans",Text="第三人称模板")
+LocalizedDescriptions=(Language="zh-Hans",Text="此模板包含一个可操控的角色，镜头采用过肩视角进行拍摄。玩家可以通过键盘、鼠标、手柄或触控设备上的虚拟摇杆来旋转镜头并进行移动。你可以从“变体（Variants）”下拉菜单中选择特定游戏类型的起始点。C++版本包含了所有变体的代码和内容。")
+
 ; This project uses a thumbnail as the project icon
 bThumbnailAsIcon=true
 
@@ -47,10 +56,10 @@ EditDetailLevelPreference="High"
 ;SharedContentPacks=(MountName="Variant_SideScroller",DetailLevels=("Standard"))
 
 ; Variant definitions
-Variants=(Name="Combat",LocalizedDisplayNames=((Language="en",Text="Combat")),LocalizedDescriptions=((Language="en",Text="A combat game with melee attacks and AI-controlled enemies.")),SharedContentPacks=((DetailLevels=(Standard),MountName="Variant_Combat")))
+Variants=(Name="Combat",LocalizedDisplayNames=((Language="en",Text="Combat"),(Language="ja",Text="戦闘"),(Language="ko",Text="전투"),(Language="zh-Hans",Text="格斗游戏模板")),LocalizedDescriptions=((Language="en",Text="A combat game with melee attacks and AI-controlled enemies."),(Language="ja",Text="近接攻撃と AI 制御の敵が登場する戦闘ゲームです"),(Language="ko",Text="근접 공격과 AI로 제어되는 적을 상대하는 전투 게임입니다."),(Language="zh-Hans",Text="具有近战攻击功能和AI控制的敌人的格斗游戏。")),SharedContentPacks=((DetailLevels=(Standard),MountName="Variant_Combat")))
 
-Variants=(Name="Platforming",LocalizedDisplayNames=((Language="en",Text="Platforming")),LocalizedDescriptions=((Language="en",Text="A 3D platforming game with tweaked movement properties and advanced jump capabilities.")),SharedContentPacks=((DetailLevels=(Standard),MountName="Variant_Platforming")))
+Variants=(Name="Platforming",LocalizedDisplayNames=((Language="en",Text="Platforming"),(Language="ja",Text="プラットフォーム"),(Language="ko",Text="플랫포머"),(Language="zh-Hans",Text="平台跳跃游戏模板")),LocalizedDescriptions=((Language="en",Text="A 3D platforming game with tweaked movement properties and advanced jump capabilities."),(Language="ja",Text="微調整された移動プロパティと高度なジャンプ機能を備えた 3D プラットフォーム ゲーム"),(Language="ko",Text="역동적인 무브먼트와 고급 점프 기능을 조정한 3D 플랫폼 게임입니다."),(Language="zh-Hans",Text="一款3D平台跳跃类游戏，其移动属性经过了调整，还具备先进的跳跃功能。")),SharedContentPacks=((DetailLevels=(Standard),MountName="Variant_Platforming")))
 
-Variants=(Name="SideScroller",LocalizedDisplayNames=((Language="en",Text="Side Scroller")),LocalizedDescriptions=((Language="en",Text="A 2.5D platforming game with plane-constrained movement, advanced jump capabilities and a side-scrolling camera.")),SharedContentPacks=((DetailLevels=(Standard),MountName="Variant_SideScroller")))
+Variants=(Name="SideScroller",LocalizedDisplayNames=((Language="en",Text="Side Scroller"),(Language="ja",Text="横スクロール"),(Language="ko",Text="횡스크롤"),(Language="zh-Hans",Text="横版卷轴游戏模板")),LocalizedDescriptions=((Language="en",Text="A 2.5D platforming game with plane-constrained movement, advanced jump capabilities and a side-scrolling camera."),(Language="ja",Text="制約のある平面移動、高度なジャンプ機能、横スクロール カメラを備えた 2.5D プラットフォーム ゲーム"),(Language="ko",Text="평면으로 제한된 움직임, 고급 점프 기능, 횡스크롤 카메라를 갖춘 2.5D 플랫폼 게임입니다."),(Language="zh-Hans",Text="一款采用2.5D视角的平台跳跃游戏，具有受平面限制的移动方式、先进的跳跃功能以横向滚动的摄像机视角。")),SharedContentPacks=((DetailLevels=(Standard),MountName="Variant_SideScroller")))
 
-Variants=(Name="All",LocalizedDisplayNames=((Language="en",Text="All")),LocalizedDescriptions=((Language="en",Text="Includes all Third Person variants.")),SharedContentPacks=((DetailLevels=(Standard),MountName="Variant_Combat"),(DetailLevels=(Standard),MountName="Variant_Platforming"),(DetailLevels=(Standard),MountName="Variant_SideScroller")))
\ No newline at end of file
+Variants=(Name="All",LocalizedDisplayNames=((Language="en",Text="All"),(Language="ja",Text="All (全員)"),(Language="ko",Text="모두"),(Language="zh-Hans",Text="全部")),LocalizedDescriptions=((Language="en",Text="Includes all Third Person variants."),(Language="ja",Text="すべての三人称バリアントを含みます"),(Language="ko",Text="모든 3인칭 베리언트를 포함합니다."),(Language="zh-Hans",Text="包含全部第三人称模板变体。")),SharedContentPacks=((DetailLevels=(Standard),MountName="Variant_Combat"),(DetailLevels=(Standard),MountName="Variant_Platforming"),(DetailLevels=(Standard),MountName="Variant_SideScroller")))
\ No newline at end of file
diff --git a/Templates/TP_ThirdPersonBP/manifest.json b/Templates/TP_ThirdPersonBP/manifest.json
--- a/Templates/TP_ThirdPersonBP/manifest.json
+++ b/Templates/TP_ThirdPersonBP/manifest.json
@@ -1,4 +1,4 @@
-{
+﻿{
 	"Version": 1,
 	"Name": [
 		{
@@ -7,15 +7,15 @@
 		},
 		{
 			"Language": "ko",
-			"Text": "삼인칭"
+			"Text": "3인칭"
 		},
 		{
 			"Language": "zh-Hans",
-			"Text": "第三人称游戏"
+			"Text": "第三人称模板"
 		},
 		{
 			"Language": "ja",
-			"Text": "サードパーソン"
+			"Text": "三人称"
 		}
 	],
 	"Description": [
@@ -25,15 +25,15 @@
 		},
 		{
 			"Language": "ko",
-			"Text": "삼인칭 템플릿에는 플레이 가능한 캐릭터가 포함되어 있으며 카메라가 캐릭터 머리 위 약간 뒤쪽에 위치합니다. 마우스, 키보드, 터치 디바이스의 가상 조이스틱으로 캐릭터가 이동함에 따라 카메라가 따라 움직이며 어깨 위 시점으로 봅니다. 이 시점은 주인공 캐릭터를 강조해 액션이나 어드벤처 게임에 주로 활용됩니다. 이 템플릿 맵에는 플레이어가 점프하여 착지하고 밀 수 있는 여러가지 오브젝트가 포함되어 있습니다."
+			"Text": "이 팩에는 카메라가 어깨 너머 시점으로 관찰하는 플레이 가능한 캐릭터가 등장합니다. 플레이어는 키보드, 마우스, 컨트롤러 또는 터치 디바이스의 가상 조이스틱을 사용하여 카메라를 회전하고 이동할 수 있습니다."
 		},
 		{
 			"Language": "zh-Hans",
-			"Text": "第三人称游戏内容包包含一个可操控的角色以及一台位于角色身后斜上方的摄像机。你可以通过鼠标、键盘、手柄或触摸式设备上的虚拟摇杆控制角色移动。摄像机会跟随角色，以越肩视角拍摄场景。这种视角突出主角，常用于动作及冒险类游戏。该模板的地图中还包含一些对象，可以让角色跳上跳下或四处推动。"
+			"Text": "此模板包含一个可操控的角色，镜头采用过肩视角进行拍摄。玩家可以通过键盘、鼠标、手柄或触控设备上的虚拟摇杆来旋转镜头并进行移动。"
 		},
 		{
 			"Language": "ja",
-			"Text": "サードパーソン パックは、後ろ側と少し上にカメラが配置されているプレイ可能なキャラクターを備えています。タッチデバイスでマウス、キーボード、コントローラー、またはバーチャル ジョイスティックの使用によりキャラクターが移動すると、カメラはキャラクターを追跡し、肩越しの視点から観察します。この視点では主人公が強調されるため、アクションゲームやアドベンチャー ゲームで人気があります。テンプレート マップには、プレイヤーがジャンプしたり押すことができるオブジェクトがいくつか含まれています。"
+			"Text": "このパックには、カメラが肩越しの視点で表示されるプレイアブル キャラクターが含まれています。プレイヤーは、キーボード、マウス、コントローラー、またはタッチ デバイスのバーチャル ジョイスティックを使用して、カメラを回転させたり移動したりできます"
 		}
 	],
 	"AssetTypes": [
diff --git a/Templates/TP_TopDown/Config/TemplateDefs.ini b/Templates/TP_TopDown/Config/TemplateDefs.ini
--- a/Templates/TP_TopDown/Config/TemplateDefs.ini
+++ b/Templates/TP_TopDown/Config/TemplateDefs.ini
@@ -4,16 +4,16 @@ Categories=Games
 
 
 LocalizedDisplayNames=(Language="en",Text="Top Down")
-LocalizedDescriptions=(Language="en",Text="This template template features a playable character and a top-down camera perspective at a greater distance. The player can control the character by clicking or touching on a destination, or following the cursor if input is held. You can select a genre-specific starting point from the Variants drop-down. The C++ version includes code and content for all variants.")
+LocalizedDescriptions=(Language="en",Text="This template features a playable character and a top-down camera perspective at a greater distance. The player can control the character by clicking or touching on a destination, or following the cursor if input is held. You can select a genre-specific starting point from the Variants drop-down. The C++ version includes code and content for all variants.")
 
-LocalizedDisplayNames=(Language="ko",Text="내려보기")
-LocalizedDescriptions=(Language="ko",Text="내려보기 템플릿은 플레이어 캐릭터 위 삼인칭 시점보다 멀리 카메라가 위치합니다. 플레이어 제어는 마우스나 터치스크린으로 원하는 목적지를 클릭하는 식으로 이루어지며, 원활한 캐릭터 이동을 위해 내비게이션 시스템을 사용합니다. 이 시점은 액션 롤플레잉 게임에 자주 쓰입니다.")
+LocalizedDisplayNames=(Language="ko",Text="탑 뷰")
+LocalizedDescriptions=(Language="ko",Text="이 템플릿은 플레이 가능한 캐릭터와 더 먼 거리에서 내려다보는 탑다운 카메라 시점을 제공합니다. 플레이어는 목적지를 클릭하거나 터치하여 캐릭터를 제어할 수 있으며, 커서를 길게 누르면 커서를 따라 움직일 수 있습니다. 베리언트 드롭다운에서 장르에 맞게 시작할 옵션을 선택할 수 있습니다. C++ 버전에는 모든 배리언트에 대한 코드와 콘텐츠가 포함되어 있습니다.")
 
-LocalizedDisplayNames=(Language="ja",Text="トップダウン")
-LocalizedDescriptions=(Language="ja",Text="トップダウンテンプレートにはプレイアブルキャラクタが存在し、カメラはキャラクタの背後、サードパーソンテンプレートより遠い上方にあります。マウスやタッチスクリーンで行き先をクリックすることでキャラクタを操作することができます。キャラクタはナビゲーションシステムを利用して移動します。この視点はアクションロールプレイングゲームでよく使用されています。")
+LocalizedDisplayNames=(Language="ja",Text="トップ ダウン")
+LocalizedDescriptions=(Language="ja",Text="このテンプレートには、プレイアブル キャラクターと、より遠距離からのトップダウン カメラの視点が含まれています。プレイヤーは、目的地をクリックまたはタッチするか、入力が保持されている場合はカーソルを追うことでキャラクターを操作できます。[バリアント] ドロップダウンからジャンル固有の開始点を選択できます。C++ バージョンには、すべてのバリアントのコードとコンテンツが含まれています")
 
-LocalizedDisplayNames=(Language="zh-Hans",Text="俯视角游戏")
-LocalizedDescriptions=(Language="zh-Hans",Text="俯视角游戏模板拥有一个可游玩角色，摄像机位于角色的后上方，但位置比第三人称视角要高得多。玩家可以使用鼠标或触摸屏点击目的地，或使用导航系统来辅助角色移动。此视角常用于动作角色扮演类游戏。")
+LocalizedDisplayNames=(Language="zh-Hans",Text="俯视角模板")
+LocalizedDescriptions=(Language="zh-Hans",Text="此模板模板包含一个可操控的角色以及一个位于较远位置的俯视视角。玩家可以通过点击或触摸目标位置来控制角色，或者使其在持续按下输入键的情况下跟随光标移动。你可以从“变体（Variants）”下拉菜单中选择特定游戏类型的起始点。C++版本包含了所有变体的代码和内容。")
 
 ; This project uses a thumbnail as the project icon
 bThumbnailAsIcon=true
diff --git a/Templates/TP_TopDown/manifest.json b/Templates/TP_TopDown/manifest.json
--- a/Templates/TP_TopDown/manifest.json
+++ b/Templates/TP_TopDown/manifest.json
@@ -1,4 +1,4 @@
-{
+﻿{
 	"Version": 1,
 	"Name": [
 		{
@@ -7,15 +7,15 @@
 		},
 		{
 			"Language": "ko",
-			"Text": "내려보기"
+			"Text": "탑 뷰"
 		},
 		{
 			"Language": "zh-Hans",
-			"Text": "俯视角游戏"
+			"Text": "俯视角模板"
 		},
 		{
 			"Language": "ja",
-			"Text": "トップダウン"
+			"Text": "トップ ダウン"
 		}
 	],
 	"Description": [
@@ -25,15 +25,15 @@
 		},
 		{
 			"Language": "ko",
-			"Text": "내려보기 템플릿은 플레이어 캐릭터 위 삼인칭 시점보다 멀리 카메라가 위치합니다. 플레이어 제어는 마우스나 터치스크린으로 원하는 목적지를 클릭하는 식으로 이루어지며, 원활한 캐릭터 이동을 위해 내비게이션 시스템을 사용합니다. 이 시점은 액션 롤플레잉 게임에 자주 쓰입니다."
+			"Text": "이 팩에는 플레이 가능한 캐릭터와 더 먼 거리에서 내려다보는 탑 뷰 카메라 시점이 특징입니다. 플레이어는 목적지를 클릭하거나 터치하여 캐릭터를 제어할 수 있으며, 커서를 길게 누르면 커서를 따라 움직일 수 있습니다."
 		},
 		{
 			"Language": "zh-Hans",
-			"Text": "俯视角游戏模板拥有一个可游玩角色，摄像机位于角色的后上方，但位置比第三人称视角要高得多。玩家可以使用鼠标或触摸屏点击目的地，或使用导航系统来辅助角色移动。此视角常用于动作角色扮演类游戏。"
+			"Text": "此模板模板包含一个可操控的角色以及一个位于较远位置的俯视视角。玩家可以通过点击或触摸目标位置来控制角色，或者使其在持续按下输入键的情况下跟随光标移动。"
 		},
 		{
 			"Language": "ja",
-			"Text": "トップダウンテンプレートにはプレイアブルキャラクタが存在し、カメラはキャラクタの背後、サードパーソンテンプレートより遠い上方にあります。マウスやタッチスクリーンで行き先をクリックすることでキャラクタを操作することができます。キャラクタはナビゲーションシステムを利用して移動します。この視点はアクションロールプレイングゲームでよく使用されています。"
+			"Text": "このパックには、プレイアブル キャラクターと、より遠距離からのトップダウン カメラの視点が含まれています。プレイヤーは、目的地をクリックまたはタッチするか、入力が保持されている場合はカーソルを追うことでキャラクターを操作できます"
 		}
 	],
 	"AssetTypes": [
diff --git a/Templates/TP_TopDownBP/Config/TemplateDefs.ini b/Templates/TP_TopDownBP/Config/TemplateDefs.ini
--- a/Templates/TP_TopDownBP/Config/TemplateDefs.ini
+++ b/Templates/TP_TopDownBP/Config/TemplateDefs.ini
@@ -4,7 +4,16 @@ Categories=Games
 
 
 LocalizedDisplayNames=(Language="en",Text="Top Down")
-LocalizedDescriptions=(Language="en",Text="This template template features a playable character and a top-down camera perspective at a greater distance. The player can control the character by clicking or touching on a destination, or following the cursor if input is held. You can select a genre-specific starting point from the Variants drop-down. The C++ version includes code and content for all variants.")
+LocalizedDescriptions=(Language="en",Text="This template features a playable character and a top-down camera perspective at a greater distance. The player can control the character by clicking or touching on a destination, or following the cursor if input is held. You can select a genre-specific starting point from the Variants drop-down. The C++ version includes code and content for all variants.")
+
+LocalizedDisplayNames=(Language="ko",Text="탑 뷰")
+LocalizedDescriptions=(Language="ko",Text="이 템플릿은 플레이 가능한 캐릭터와 더 먼 거리에서 내려다보는 탑다운 카메라 시점을 제공합니다. 플레이어는 목적지를 클릭하거나 터치하여 캐릭터를 제어할 수 있으며, 커서를 길게 누르면 커서를 따라 움직일 수 있습니다. 베리언트 드롭다운에서 장르에 맞게 시작할 옵션을 선택할 수 있습니다. C++ 버전에는 모든 배리언트에 대한 코드와 콘텐츠가 포함되어 있습니다.")
+
+LocalizedDisplayNames=(Language="ja",Text="トップ ダウン")
+LocalizedDescriptions=(Language="ja",Text="このテンプレートには、プレイアブル キャラクターと、より遠距離からのトップダウン カメラの視点が含まれています。プレイヤーは、目的地をクリックまたはタッチするか、入力が保持されている場合はカーソルを追うことでキャラクターを操作できます。[バリアント] ドロップダウンからジャンル固有の開始点を選択できます。C++ バージョンには、すべてのバリアントのコードとコンテンツが含まれています")
+
+LocalizedDisplayNames=(Language="zh-Hans",Text="俯视角模板")
+LocalizedDescriptions=(Language="zh-Hans",Text="此模板模板包含一个可操控的角色以及一个位于较远位置的俯视视角。玩家可以通过点击或触摸目标位置来控制角色，或者使其在持续按下输入键的情况下跟随光标移动。你可以从“变体（Variants）”下拉菜单中选择特定游戏类型的起始点。C++版本包含了所有变体的代码和内容。")
 
 ; This project uses a thumbnail as the project icon
 bThumbnailAsIcon=true
@@ -48,8 +57,8 @@ EditDetailLevelPreference="High"
 ;SharedContentPacks=(MountName="Input",DetailLevels=("High"))
 
 ; Variant definitions
-Variants=(Name="Strategy",LocalizedDisplayNames=((Language="en",Text="Strategy")),LocalizedDescriptions=((Language="en",Text="A top down strategy game with point and click interactions and selectable units.")),SharedContentPacks=((DetailLevels=(Standard),MountName="Variant_Strategy")))
+Variants=(Name="Strategy",LocalizedDisplayNames=((Language="en",Text="Strategy"),(Language="ja",Text="戦略"),(Language="ko",Text="전략"),(Language="zh-Hans",Text="战略游戏模板")),LocalizedDescriptions=((Language="en",Text="A top down strategy game with point and click interactions and selectable units."),(Language="ja",Text="ポイント アンド クリック操作と選択可能なユニットを備えたトップダウン戦略ゲーム"),(Language="ko",Text="포인트 앤 클릭 상호작용과 선택 가능한 유닛이 포함된 탑 뷰 전략 게임입니다."),(Language="zh-Hans",Text="一款俯视角战略游戏，采用点击操作并具有可选择单位。")),SharedContentPacks=((DetailLevels=(Standard),MountName="Variant_Strategy")))
 
-Variants=(Name="TwinStick",LocalizedDisplayNames=((Language="en",Text="Twin Stick")),LocalizedDescriptions=((Language="en",Text="A twin stick shooter game with infinitely spawning enemies.")),SharedContentPacks=((DetailLevels=(Standard),MountName="Variant_TwinStick"),(MountName="Input",DetailLevels=("High")))
+Variants=(Name="TwinStick",LocalizedDisplayNames=((Language="en",Text="Twin Stick"),(Language="ja",Text="ツインスティック"),(Language="ko",Text="트윈 스틱"),(Language="zh-Hans",Text="双摇杆游戏模板")),LocalizedDescriptions=((Language="en",Text="A twin stick shooter game with infinitely spawning enemies."),(Language="ja",Text="敵が無限にスポーンするツイン スティック シューティング ゲーム"),(Language="ko",Text="무한히 생성되는 적을 상대하는 트윈 스틱 슈팅 게임입니다."),(Language="zh-Hans",Text="一款双摇杆设计游戏，具有无限刷新的敌人。")),SharedContentPacks=((DetailLevels=(Standard),MountName="Variant_TwinStick"),(MountName="Input",DetailLevels=("High")))
 
-Variants=(Name="All",LocalizedDisplayNames=((Language="en",Text="All")),LocalizedDescriptions=((Language="en",Text="Includes all Top Down variants.")),SharedContentPacks=((DetailLevels=(Standard),MountName="Variant_TwinStick"),(DetailLevels=(Standard),MountName="Variant_Strategy"),(MountName="Input",DetailLevels=("High")))
\ No newline at end of file
+Variants=(Name="All",LocalizedDisplayNames=((Language="en",Text="All"),(Language="ja",Text="All (全員)"),(Language="ko",Text="모두"),(Language="zh-Hans",Text="全部")),LocalizedDescriptions=((Language="en",Text="Includes all Top Down variants."),(Language="ja",Text="すべてのトップダウン バリアントが含まれています"),(Language="ko",Text="모든 탑 뷰 베리언트를 포함합니다."),(Language="zh-Hans",Text="包含全部第三人称模板变体。")),SharedContentPacks=((DetailLevels=(Standard),MountName="Variant_TwinStick"),(DetailLevels=(Standard),MountName="Variant_Strategy"),(MountName="Input",DetailLevels=("High")))
\ No newline at end of file
diff --git a/Templates/TP_TopDownBP/manifest.json b/Templates/TP_TopDownBP/manifest.json
--- a/Templates/TP_TopDownBP/manifest.json
+++ b/Templates/TP_TopDownBP/manifest.json
@@ -1,4 +1,4 @@
-{
+﻿{
 	"Version": 1,
 	"Name": [
 		{
@@ -7,15 +7,15 @@
 		},
 		{
 			"Language": "ko",
-			"Text": "탑다운"
+			"Text": "탑 뷰"
 		},
 		{
 			"Language": "zh-Hans",
-			"Text": "俯视角游戏"
+			"Text": "俯视角模板"
 		},
 		{
 			"Language": "ja",
-			"Text": "トップダウン"
+			"Text": "トップ ダウン"
 		}
 	],
 	"Description": [
@@ -25,15 +25,15 @@
 		},
 		{
 			"Language": "ko",
-			"Text": "탑다운 템플릿에는 플레이 가능한 캐릭터가 포함되어 있으며 카메라가 아바타 머리 위 약간 뒤쪽에서 삼인칭 시점보다 더 원거리에 위치합니다. 플레이어는 마우스, 게임패드나 터치스크린으로 제어되고, 목적지를 클릭하거나 입력을 유지하여 커서에 따라 캐릭터를 움직일 수 있습니다. 이 시점은 주로 액션 RPG에 활용됩니다."
+			"Text": "이 팩에는 플레이 가능한 캐릭터와 더 먼 거리에서 내려다보는 탑 뷰 카메라 시점이 특징입니다. 플레이어는 목적지를 클릭하거나 터치하여 캐릭터를 제어할 수 있으며, 커서를 길게 누르면 커서를 따라 움직일 수 있습니다."
 		},
 		{
 			"Language": "zh-Hans",
-			"Text": "俯视角游戏模板包含一个可操控的角色以及一台位于角色后上方的摄像机。摄像机的位置比第三人称视角模板中的要高很多。玩家可以通过点击鼠标、触控板或触摸屏来选择目的地并移动角色，或持续按下鼠标来移动角色。此视角常用于动作角色扮演类游戏。"
+			"Text": "此模板模板包含一个可操控的角色以及一个位于较远位置的俯视视角。玩家可以通过点击或触摸目标位置来控制角色，或者使其在持续按下输入键的情况下跟随光标移动。"
 		},
 		{
 			"Language": "ja",
-			"Text": "トップダウン テンプレートは、カメラがアバターの後ろと上に配置されている (サードパーソン視点よりも距離が遠い位置) プレイ可能なキャラクターを特徴としています。プレイヤーは、マウス、ゲームパッド、またはタッチ スクリーンの使用により制御され、目的地をクリックするか、入力が保持されている場合はカーソルをたどることでキャラクターを移動させます。この視点は、アクション ロールプレイ ゲームでよく使用されます。"
+			"Text": "このパックには、プレイアブル キャラクターと、より遠距離からのトップダウン カメラの視点が含まれています。プレイヤーは、目的地をクリックまたはタッチするか、入力が保持されている場合はカーソルを追うことでキャラクターを操作できます"
 		}
 	],
 	"AssetTypes": [
diff --git a/Templates/TP_VehicleAdv/Config/TemplateDefs.ini b/Templates/TP_VehicleAdv/Config/TemplateDefs.ini
--- a/Templates/TP_VehicleAdv/Config/TemplateDefs.ini
+++ b/Templates/TP_VehicleAdv/Config/TemplateDefs.ini
@@ -6,13 +6,13 @@ LocalizedDisplayNames=(Language="en",Text="Vehicle")
 LocalizedDescriptions=(Language="en",Text="This template features a sports car and an offroad vehicle with a double wishbone suspension, as well as a HUD displaying the current gear and speed. Vehicle movement can be controlled with keyboard, controller, steering wheel, or virtual joystick on a touch device. You can select a genre-specific starting point from the Variants drop-down. The C++ version includes code and content for all variants.")
 
 LocalizedDisplayNames=(Language="ko",Text="비히클")
-LocalizedDescriptions=(Language="ko",Text="비히클 팩은 일반 비히클과 조금 더 복잡한 비히클을 모두 제공합니다. 비히클의 이동은 키보드, 컨트롤러 또는 터치 디바이스의 버추얼 조이스틱으로 제어 가능합니다. 고급 비히클은 더블 위시본 서스펜션을 갖추고 있습니다. 이 템플릿에는 엔진 사운드와, 현재 기어와 속도를 보여주는 HUD도 포함되어 있습니다.")
+LocalizedDescriptions=(Language="ko",Text="이 템플릿에는 스포츠카와 더블 위시본 서스펜션이 장착된 오프로드 차량, 그리고 현재 기어와 속도를 표시하는 HUD가 있습니다. 차량의 움직임은 키보드, 컨트롤러, 스티어링 휠 또는 터치 디바이스의 가상 조이스틱으로 제어할 수 있습니다. 배리언트 드롭다운에서 장르에 맞게 시작할 옵션을 선택할 수 있습니다. C++ 버전에는 모든 배리언트에 대한 코드와 콘텐츠가 포함되어 있습니다.")
 
 LocalizedDisplayNames=(Language="ja",Text="ビークル")
-LocalizedDescriptions=(Language="ja",Text="ビークル パックには通常のビークルとより複雑なビークルが含まれています。キーボード、コントローラー、タッチデバイスでのバーチャルジョイスティックでビークルを操作できます。このビークルにはダブルウィッシュボーン式サスペンションが含まれています。またエンジン音もテンプレートにあります。また現在のギアと速度を示すHUDも用意されています。")
+LocalizedDescriptions=(Language="ja",Text="このテンプレートには、スポーツ カーとダブル ウィッシュボーン サスペンションを備えたオフロード車、そして現在のギアと速度を表示する HUD が含まれています。車の動きは、キーボード、コントローラー、ステアリング ホイール、またはタッチ デバイスのバーチャル ジョイスティックで操作できます。[バリアント] ドロップダウンからジャンル固有の開始点を選択できます。C++ バージョンには、すべてのバリアントのコードとコンテンツが含まれています")
 
-LocalizedDisplayNames=(Language="zh-Hans",Text="载具")
-LocalizedDescriptions=(Language="zh-Hans",Text="载具包包含一辆普通载具和高级载具。你可以通过键盘、手柄或触摸式设备上的虚拟摇杆来控制载具移动。高级载具具有双横臂式悬架。此外，该模板还包含引擎音效，以及一个能够显示当前挡位和车速的HUD。")
+LocalizedDisplayNames=(Language="zh-Hans",Text="载具模板")
+LocalizedDescriptions=(Language="zh-Hans",Text="此模板展示了一辆跑车和一辆配备双叉臂悬架的越野车，同时还有一块抬头显示器，可显示当前挡位和速度。载具的移动可以通过键盘、手柄、方向盘或触控设备上的虚拟摇杆来控你。你可以从“变体（Variants）”下拉菜单中选择特定游戏类型的起始点。C++ 版本包含了所有变体的代码和内容。")
 
 ; This project uses a thumbnail as the project icon
 bThumbnailAsIcon=true
diff --git a/Templates/TP_VehicleAdv/manifest.json b/Templates/TP_VehicleAdv/manifest.json
--- a/Templates/TP_VehicleAdv/manifest.json
+++ b/Templates/TP_VehicleAdv/manifest.json
@@ -1,4 +1,4 @@
-{
+﻿{
 	"Version": 1,
 	"Name":
 	[
@@ -12,7 +12,7 @@
 		},
 		{
 			"Language": "zh-Hans",
-			"Text": "载具"
+			"Text": "载具模板"
 		},
 		{
 			"Language": "ko",
@@ -27,15 +27,15 @@
 		},
 		{
 			"Language": "ja",
-			"Text": "ビークル パックには通常のビークルとより複雑なビークルが含まれています。キーボード、コントローラー、タッチデバイスでのバーチャルジョイスティックでビークルを操作できます。このビークルにはダブルウィッシュボーン式サスペンションが含まれています。またエンジン音もテンプレートにあります。また現在のギアと速度を示すHUDも用意されています。"
+			"Text": "このパックには、スポーツ カーとダブル ウィッシュボーン サスペンションを備えたオフロード車、そして現在のギアと速度を表示する HUD が含まれています。車の動きは、キーボード、コントローラー、ステアリング ホイール、またはタッチ デバイスのバーチャル ジョイスティックで操作できます"
 		},
 		{
 			"Language": "zh-Hans",
-			"Text": "载具包包含一辆普通载具和高级载具。你可以通过键盘、手柄或触摸式设备上的虚拟摇杆来控制载具移动。高级载具具有双横臂式悬架。此外，该模板还包含引擎音效，以及一个能够显示当前挡位和车速的HUD。"
+			"Text": "此模板展示了一辆跑车和一辆配备双叉臂悬架的越野车，同时还有一块抬头显示器，可显示当前挡位和速度。载具的移动可以通过键盘、手柄、方向盘或触控设备上的虚拟摇杆来控你。"
 		},
 		{
 			"Language": "ko",
-			"Text": "비히클 팩은 일반 비히클과 조금 더 복잡한 비히클을 모두 제공합니다. 비히클의 이동은 키보드, 컨트롤러 또는 터치 디바이스의 버추얼 조이스틱으로 제어 가능합니다. 고급 비히클은 더블 위시본 서스펜션을 갖추고 있습니다. 이 템플릿에는 엔진 사운드와, 현재 기어와 속도를 보여주는 HUD도 포함되어 있습니다."
+			"Text": "이 팩에는 스포츠카와 더블 위시본 서스펜션이 장착된 오프로드 차량, 그리고 현재 기어와 속도를 표시하는 HUD가 포함되어 있습니다. 키보드, 컨트롤러, 스티어링 휠 또는 터치 디바이스의 가상 조이스틱으로 차량의 움직임을 제어할 수 있습니다."
 		}
 	],
 	"AssetTypes":
diff --git a/Templates/TP_VehicleAdvBP/Config/TemplateDefs.ini b/Templates/TP_VehicleAdvBP/Config/TemplateDefs.ini
--- a/Templates/TP_VehicleAdvBP/Config/TemplateDefs.ini
+++ b/Templates/TP_VehicleAdvBP/Config/TemplateDefs.ini
@@ -5,6 +5,15 @@ Categories=Games
 LocalizedDisplayNames=(Language="en",Text="Vehicle")
 LocalizedDescriptions=(Language="en",Text="This template features a sports car and an offroad vehicle with a double wishbone suspension, as well as a HUD displaying the current gear and speed. Vehicle movement can be controlled with keyboard, controller, steering wheel, or virtual joystick on a touch device. You can select a genre-specific starting point from the Variants drop-down. The C++ version includes code and content for all variants.")
 
+LocalizedDisplayNames=(Language="ko",Text="비히클")
+LocalizedDescriptions=(Language="ko",Text="이 템플릿에는 스포츠카와 더블 위시본 서스펜션이 장착된 오프로드 차량, 그리고 현재 기어와 속도를 표시하는 HUD가 있습니다. 차량의 움직임은 키보드, 컨트롤러, 스티어링 휠 또는 터치 디바이스의 가상 조이스틱으로 제어할 수 있습니다. 배리언트 드롭다운에서 장르에 맞게 시작할 옵션을 선택할 수 있습니다. C++ 버전에는 모든 배리언트에 대한 코드와 콘텐츠가 포함되어 있습니다.")
+
+LocalizedDisplayNames=(Language="ja",Text="ビークル")
+LocalizedDescriptions=(Language="ja",Text="このテンプレートには、スポーツ カーとダブル ウィッシュボーン サスペンションを備えたオフロード車、そして現在のギアと速度を表示する HUD が含まれています。車の動きは、キーボード、コントローラー、ステアリング ホイール、またはタッチ デバイスのバーチャル ジョイスティックで操作できます。[バリアント] ドロップダウンからジャンル固有の開始点を選択できます。C++ バージョンには、すべてのバリアントのコードとコンテンツが含まれています")
+
+LocalizedDisplayNames=(Language="zh-Hans",Text="载具模板")
+LocalizedDescriptions=(Language="zh-Hans",Text="此模板展示了一辆跑车和一辆配备双叉臂悬架的越野车，同时还有一块抬头显示器，可显示当前挡位和速度。载具的移动可以通过键盘、手柄、方向盘或触控设备上的虚拟摇杆来控你。你可以从“变体（Variants）”下拉菜单中选择特定游戏类型的起始点。C++ 版本包含了所有变体的代码和内容。")
+
 ; This project uses a thumbnail as the project icon
 bThumbnailAsIcon=true
 
@@ -45,8 +54,8 @@ EditDetailLevelPreference="High"
 ;SharedContentPacks=(MountName="Variant_Offroad",DetailLevels=("Standard"))
 ;SharedContentPacks=(MountName="Variant_Timetrial",DetailLevels=("Standard"))
 
-Variants=(Name="Timetrial",LocalizedDisplayNames=((Language="en",Text="Time Trial")),LocalizedDescriptions=((Language="en",Text="A timed lap with UI and track markers.")),SharedContentPacks=((DetailLevels=(Standard),MountName="Variant_Timetrial")))
+Variants=(Name="Timetrial",LocalizedDisplayNames=((Language="en",Text="Time Trial"),(Language="ja",Text="タイム トライアル"),(Language="ko",Text="타임 트라이얼"),(Language="zh-Hans",Text="计时赛模板")),LocalizedDescriptions=((Language="en",Text="A timed lap with UI and track markers."),(Language="ja",Text="UI とトラック マーカーを備えたタイムラップ"),(Language="ko",Text="UI와 트랙 마커가 있는 시간 제한 랩입니다."),(Language="zh-Hans",Text="带有用户界面和赛道标记的计时圈赛。")),SharedContentPacks=((DetailLevels=(Standard),MountName="Variant_Timetrial")))
 
-Variants=(Name="Offroad",LocalizedDisplayNames=((Language="en",Text="Offroad")),LocalizedDescriptions=((Language="en",Text="An open world map with an advanced offroad vehicle.")),SharedContentPacks=((DetailLevels=(Standard),MountName="Variant_Offroad")))
+Variants=(Name="Offroad",LocalizedDisplayNames=((Language="en",Text="Offroad"),(Language="ja",Text="オフロード"),(Language="ko",Text="오프로드"),(Language="zh-Hans",Text="越野车模板")),LocalizedDescriptions=((Language="en",Text="An open world map with an advanced offroad vehicle."),(Language="ja",Text="高度なオフロード車を備えたオープン ワールド マップ"),(Language="ko",Text="고급 오프로드 차량이 있는 오픈 월드 맵입니다."),(Language="zh-Hans",Text="一张开放世界地图和一辆先进的越野车。")),SharedContentPacks=((DetailLevels=(Standard),MountName="Variant_Offroad")))
 
-Variants=(Name="All",LocalizedDisplayNames=((Language="en",Text="All")),LocalizedDescriptions=((Language="en",Text="Includes all Vehicle variants.")),SharedContentPacks=((DetailLevels=(Standard),MountName="Variant_Offroad"),(DetailLevels=(Standard),MountName="Variant_Timetrial")))
\ No newline at end of file
+Variants=(Name="All",LocalizedDisplayNames=((Language="en",Text="All"),(Language="ja",Text="All (全員)"),(Language="ko",Text="모두"),(Language="zh-Hans",Text="全部")),LocalizedDescriptions=((Language="en",Text="Includes all Vehicle variants."),(Language="ja",Text="すべてのビークル バリアントを含みます"),(Language="ko",Text="모든 비히클 베리언트를 포함합니다."),(Language="zh-Hans",Text="包含全部载具模板变体。")),SharedContentPacks=((DetailLevels=(Standard),MountName="Variant_Offroad"),(DetailLevels=(Standard),MountName="Variant_Timetrial")))
\ No newline at end of file
diff --git a/Templates/TP_VehicleAdvBP/manifest.json b/Templates/TP_VehicleAdvBP/manifest.json
--- a/Templates/TP_VehicleAdvBP/manifest.json
+++ b/Templates/TP_VehicleAdvBP/manifest.json
@@ -1,4 +1,4 @@
-{
+﻿{
 	"Version": 1,
 	"Name": [
 		{
@@ -11,7 +11,7 @@
 		},
 		{
 			"Language": "zh-Hans",
-			"Text": "载具"
+			"Text": "载具模板"
 		},
 		{
 			"Language": "ko",
@@ -25,15 +25,15 @@
 		},
 		{
 			"Language": "ja",
-			"Text": "ビークル パックには通常のビークルとより複雑なビークルが含まれています。キーボード、コントローラー、タッチデバイスでのバーチャルジョイスティックでビークルを操作できます。このビークルにはダブルウィッシュボーン式サスペンションが含まれています。またエンジン音もテンプレートにあります。また現在のギアと速度を示すHUDも用意されています。"
+			"Text": "このパックには、スポーツ カーとダブル ウィッシュボーン サスペンションを備えたオフロード車、そして現在のギアと速度を表示する HUD が含まれています。車の動きは、キーボード、コントローラー、ステアリング ホイール、またはタッチ デバイスのバーチャル ジョイスティックで操作できます"
 		},
 		{
 			"Language": "zh-Hans",
-			"Text": "载具包包含一辆普通载具和高级载具。你可以通过键盘、手柄或触摸式设备上的虚拟摇杆来控制载具移动。高级载具具有双横臂式悬架。此外，该模板还包含引擎音效，以及一个能够显示当前挡位和车速的HUD。"
+			"Text": "此模板展示了一辆跑车和一辆配备双叉臂悬架的越野车，同时还有一块抬头显示器，可显示当前挡位和速度。载具的移动可以通过键盘、手柄、方向盘或触控设备上的虚拟摇杆来控你。"
 		},
 		{
 			"Language": "ko",
-			"Text": "비히클 팩은 일반 비히클과 조금 더 복잡한 비히클을 모두 제공합니다. 비히클의 이동은 키보드, 컨트롤러 또는 터치 디바이스의 버추얼 조이스틱으로 제어 가능합니다. 고급 비히클은 더블 위시본 서스펜션을 갖추고 있습니다. 이 템플릿에는 엔진 사운드와, 현재 기어와 속도를 보여주는 HUD도 포함되어 있습니다."
+			"Text": "이 팩에는 스포츠카와 더블 위시본 서스펜션이 장착된 오프로드 차량, 그리고 현재 기어와 속도를 표시하는 HUD가 포함되어 있습니다. 키보드, 컨트롤러, 스티어링 휠 또는 터치 디바이스의 가상 조이스틱으로 차량의 움직임을 제어할 수 있습니다."
 		}
 	],
 	"AssetTypes": [
